Added new experimental plugin AVCodecs which handles software/hardware encoding and decoding of audio and video

diff --git a/Engine/Plugins/Experimental/AVCodecs/AVCodecsCore/Source/AVCodecsCore/Private/Video/Decoders/Configs/VideoDecoderConfigH265.cpp b/Engine/Plugins/Experimental/AVCodecs/AVCodecsCore/Source/AVCodecsCore/Private/Video/Decoders/Configs/VideoDecoderConfigH265.cpp
new file mode 100644
--- /dev/null
+++ b/Engine/Plugins/Experimental/AVCodecs/AVCodecsCore/Source/AVCodecsCore/Private/Video/Decoders/Configs/VideoDecoderConfigH265.cpp
@@ -0,0 +1,435 @@
+// Copyright Epic Games, Inc. All Rights Reserved.
+
+#include "Video/Decoders/Configs/VideoDecoderConfigH265.h"
+
+#include "AVUtility.h"
+#include "Video/CodecUtils/CodecUtilsH265.h"
+
+REGISTER_TYPEID(FVideoDecoderConfigH265);
+
+FAVResult FVideoDecoderConfigH265::Parse(TSharedRef<FAVInstance> const& Instance, FVideoPacket const& Packet, TArray<TSharedPtr<UE::AVCodecCore::H265::FNaluH265>>& Nalus)
+{
+	using namespace UE::AVCodecCore::H265;
+
+	TArray<FNaluH265> FoundNalus;
+
+	FAVResult Result;
+	Result = FindNALUs(Packet, FoundNalus);
+
+	if (Result.IsNotSuccess())
+	{
+		return Result;
+	}
+
+	FVideoDecoderConfigH265& H265 = Instance->Edit<FVideoDecoderConfigH265>();
+	TSharedPtr<FNaluVPS> CandidateVPS;
+	TSharedPtr<FNaluSPS> CandidateSPS;
+	TSharedPtr<FNaluPPS> CandidatePPS;
+
+	for (auto& Nalu : FoundNalus)
+	{
+		Result = EAVResult::Success;
+		switch (Nalu.nal_unit_type)
+		{
+			case ENaluType::TRAIL_N:
+			case ENaluType::TRAIL_R:
+			case ENaluType::TSA_N:
+			case ENaluType::TSA_R:
+			case ENaluType::STSA_N:
+			case ENaluType::STSA_R:
+			case ENaluType::RADL_N:
+			case ENaluType::RADL_R:
+			case ENaluType::RASL_N:
+			case ENaluType::RASL_R:
+				Nalus.Add(MakeShared<FNaluSlice>(Nalu));
+				StaticCastSharedPtr<FNaluSlice>(Nalus.Last())->Parse(this);
+				UE_LOG(LogTemp, Verbose, TEXT("H256 Parsing: Recieved Slice"));
+				break;
+			case ENaluType::RSV_VCL_N10:
+				checkNoEntry(); // Reserved
+				break;
+			case ENaluType::RSV_VCL_R11:
+				checkNoEntry(); // Reserved
+				break;
+			case ENaluType::RSV_VCL_N12:
+				checkNoEntry(); // Reserved
+				break;
+			case ENaluType::RSV_VCL_R13:
+				checkNoEntry(); // Reserved
+				break;
+			case ENaluType::RSV_VCL_N14:
+				checkNoEntry(); // Reserved
+				break;
+			case ENaluType::RSV_VCL_R15:
+				checkNoEntry(); // Reserved
+				break;
+			case ENaluType::IDR_W_RADL:
+			case ENaluType::IDR_N_LP:
+				UE_LOG(LogTemp, Verbose, TEXT("H256 Parsing: Recieved IDR"));
+			case ENaluType::BLA_W_LP:
+			case ENaluType::BLA_W_RADL:
+			case ENaluType::BLA_N_LP:
+			case ENaluType::CRA_NUT:
+				Nalus.Add(MakeShared<FNaluSlice>(Nalu));
+				StaticCastSharedPtr<FNaluSlice>(Nalus.Last())->Parse(this);
+				UE_LOG(LogTemp, Verbose, TEXT("H256 Parsing: Recieved Slice"));
+				break;
+			case ENaluType::RSV_IRAP_VCL22:
+				checkNoEntry(); // Reserved or Unspecified
+				break;
+			case ENaluType::RSV_IRAP_VCL23:
+				checkNoEntry(); // Reserved or Unspecified
+				break;
+			case ENaluType::RSV_VCL24:
+				checkNoEntry(); // Reserved or Unspecified
+				break;
+			case ENaluType::RSV_VCL25:
+				checkNoEntry(); // Reserved or Unspecified
+				break;
+			case ENaluType::RSV_VCL26:
+				checkNoEntry(); // Reserved or Unspecified
+				break;
+			case ENaluType::RSV_VCL27:
+				checkNoEntry(); // Reserved or Unspecified
+				break;
+			case ENaluType::RSV_VCL28:
+				checkNoEntry(); // Reserved or Unspecified
+				break;
+			case ENaluType::RSV_VCL29:
+				checkNoEntry(); // Reserved or Unspecified
+				break;
+			case ENaluType::RSV_VCL30:
+				checkNoEntry(); // Reserved or Unspecified
+				break;
+			case ENaluType::RSV_VCL31:
+				checkNoEntry(); // Reserved
+				break;
+			case ENaluType::VPS_NUT:
+				CandidateVPS = MakeShared<FNaluVPS>(Nalu);
+				CandidateVPS->Parse();
+				H265.ParsedVPS.Add(CandidateVPS->vps_video_parameter_set_id, CandidateVPS);
+				Nalus.Add(CandidateVPS);
+				UE_LOG(LogTemp, Verbose, TEXT("H256 Parsing: Recieved VPS %u"), CandidateVPS->vps_video_parameter_set_id.Value);
+				break;
+			case ENaluType::SPS_NUT:
+				CandidateSPS = MakeShared<FNaluSPS>(Nalu);
+				CandidateSPS->Parse();
+				H265.ParsedSPS.Add(CandidateSPS->sps_video_parameter_set_id, CandidateSPS);
+				Nalus.Add(CandidateSPS);
+				UE_LOG(LogTemp, Verbose, TEXT("H256 Parsing: Recieved SPS %u"), CandidateSPS->sps_video_parameter_set_id.Value);
+				break;
+			case ENaluType::PPS_NUT:
+				CandidatePPS = MakeShared<FNaluPPS>(Nalu);
+				CandidatePPS->Parse();
+				H265.ParsedPPS.Add(CandidatePPS->pps_pic_parameter_set_id, CandidatePPS);
+				Nalus.Add(CandidatePPS);
+				UE_LOG(LogTemp, Verbose, TEXT("H256 Parsing: Recieved PPS %u"), CandidatePPS->pps_pic_parameter_set_id.Value);
+				break;
+			case ENaluType::AUD_NUT:
+				Nalus.Add(MakeShared<FNaluH265>(Nalu));
+				UE_LOG(LogTemp, Verbose, TEXT("H256 Parsing: Recieved Access Unit Delimiter"));
+				break;
+			case ENaluType::EOS_NUT:
+				Nalus.Add(MakeShared<FNaluH265>(Nalu));
+				UE_LOG(LogTemp, Verbose, TEXT("H256 Parsing: Recieved End Of Sequence"));
+				break;
+			case ENaluType::EOB_NUT:
+				Nalus.Add(MakeShared<FNaluH265>(Nalu));
+				UE_LOG(LogTemp, Verbose, TEXT("H256 Parsing: Recieved End Of Bitstream"));
+				break;
+			case ENaluType::FD_NUT:
+				Nalus.Add(MakeShared<FNaluH265>(Nalu));
+				UE_LOG(LogTemp, Verbose, TEXT("H256 Parsing: Recieved Filler Data"));
+				break;
+			case ENaluType::PREFIX_SEI_NUT:
+			case ENaluType::SUFFIX_SEI_NUT:
+				H265.ParsedSEI.Add(MakeShared<FNaluSEI>(Nalu));
+				H265.ParsedSEI.Last()->Parse();
+				UE_LOG(LogTemp, Verbose, TEXT("H256 Parsing: Recieved SEI"));
+				break;
+			case ENaluType::RSV_NVCL41:
+				checkNoEntry(); // Reserved or Unspecified
+				break;
+			case ENaluType::RSV_NVCL42:
+				checkNoEntry(); // Reserved or Unspecified
+				break;
+			case ENaluType::RSV_NVCL43:
+				checkNoEntry(); // Reserved or Unspecified
+				break;
+			case ENaluType::RSV_NVCL44:
+				checkNoEntry(); // Reserved or Unspecified
+				break;
+			case ENaluType::RSV_NVCL45:
+				checkNoEntry(); // Reserved or Unspecified
+				break;
+			case ENaluType::RSV_NVCL46:
+				checkNoEntry(); // Reserved or Unspecified
+				break;
+			case ENaluType::RSV_NVCL47:
+				checkNoEntry(); // Reserved or Unspecified
+				break;
+			case ENaluType::UNSPECIFIED:
+				checkNoEntry(); // Reserved or Unspecified
+				break;
+			default:
+				checkNoEntry(); // Reserved or Unspecified
+				break;
+		}
+
+		if (Result.IsNotSuccess())
+		{
+			return Result;
+		}
+	}
+
+	return Result;
+}
+
+FAVResult FVideoDecoderConfigH265::UpdateRPS(TSharedRef<UE::AVCodecCore::H265::FNaluSlice> CurrentSlice)
+{
+	using namespace  UE::AVCodecCore::H265;
+
+	ReferencePictureSet.PicOrderCntVal[LastRpsIdx] = LastPicOrderCntValue;
+
+	// Invalidate all of the usages when a new Intra Decode Refresh frame comes in
+	if (CurrentSlice->IsIDR())
+	{
+		// Clear the current RPS working state and usage
+		ReferencePictureSet.PrepState();
+
+		LastRpsIdx = 0;
+		LastPicOrderCntValue = CurrentSlice->CurrPicOrderCntVal;
+	}
+	else
+	{		
+		TSharedPtr<FNaluSPS> PinnedSPS = CurrentSlice->CurrentSPS.Pin();
+		TSharedPtr<FNaluPPS> PinnedPPS = CurrentSlice->CurrentPPS.Pin();
+		uint8 i, j, k;
+		
+		// (8-5) Calculate NumPocTotalCur
+		uint8 const& CurrRspIdx = CurrentSlice->CurrRpsIdx;
+		short_term_ref_pic_set_t const& ShortTermRefPicSet = PinnedSPS->short_term_ref_pic_sets[CurrRspIdx]; 
+		uint32 MaxPicOrderCntLsb = 1 << (PinnedSPS->log2_max_pic_order_cnt_lsb_minus4 + 4);
+		bool UsedByCurrPicLt;
+
+		// NOTE (aidan) Should we be creating these with every slice?
+		uint8_t CurrDeltaPocMsbPresentFlag[16];
+		uint8_t FollDeltaPocMsbPresentFlag[16];
+	
+		for (i = 0, j = 0, k = 0; i < ShortTermRefPicSet.NumNegativePics; i++)
+		{
+			if (ShortTermRefPicSet.UsedByCurrPicS0[i])
+			{
+				ReferencePictureSet.PocStCurrBefore[j++] = CurrentSlice->CurrPicOrderCntVal + ShortTermRefPicSet.DeltaPocS0[i];
+			}
+			else
+			{
+				ReferencePictureSet.PocStFoll[k++] = CurrentSlice->CurrPicOrderCntVal + ShortTermRefPicSet.DeltaPocS0[i];
+			}
+		}
+		ReferencePictureSet.NumPocStCurrBefore = j;
+
+		for (i = 0, j = 0; i < ShortTermRefPicSet.NumPositivePics; i++)
+		{
+			if (ShortTermRefPicSet.UsedByCurrPicS1[i])
+			{
+				ReferencePictureSet.PocStCurrAfter[j++] = CurrentSlice->CurrPicOrderCntVal + ShortTermRefPicSet.DeltaPocS1[i];
+			}
+			else
+			{
+				ReferencePictureSet.PocStFoll[k++] = CurrentSlice->CurrPicOrderCntVal + ShortTermRefPicSet.DeltaPocS1[i];
+			}
+		}
+		ReferencePictureSet.NumPocStCurrAfter = j;
+		ReferencePictureSet.NumPocStFoll = k;
+	
+		uint8 pocLt;
+		for (i = 0, j = 0, k = 0; i < (CurrentSlice->num_long_term_sps + CurrentSlice->num_long_term_pics); i++)
+		{			
+			if (i < CurrentSlice->num_long_term_sps)
+			{
+				pocLt = PinnedSPS->long_term_ref_pics_sps[CurrentSlice->lt_idx_sps[i]].lt_ref_pic_poc_lsb_sps;
+				UsedByCurrPicLt = PinnedSPS->long_term_ref_pics_sps[CurrentSlice->lt_idx_sps[i]].used_by_curr_pic_lt_sps_flag.AsBool();
+			}
+			else
+			{
+				pocLt = CurrentSlice->poc_lsb_lt[i];
+				UsedByCurrPicLt = CurrentSlice->used_by_curr_pic_lt_flag[CurrentSlice->lt_idx_sps[i]].AsBool();
+			}
+			
+			if (CurrentSlice->delta_poc_msb_present_flag[i])
+			{
+				pocLt += CurrentSlice->CurrPicOrderCntVal - CurrentSlice->delta_poc_msb_cycle_lt[i] * MaxPicOrderCntLsb - ( CurrentSlice->CurrPicOrderCntVal & (MaxPicOrderCntLsb-1u));
+			}
+
+			if (UsedByCurrPicLt)
+			{
+				ReferencePictureSet.PocLtCurr[j] = pocLt;
+				CurrDeltaPocMsbPresentFlag[j++] = CurrentSlice->delta_poc_msb_present_flag[i];
+				ReferencePictureSet.NumPicTotalCurr++;
+			}
+			else
+			{
+				ReferencePictureSet.PocLtFoll[k] = pocLt;
+				FollDeltaPocMsbPresentFlag[k++] = CurrentSlice->delta_poc_msb_present_flag[i];
+			}
+		}
+		ReferencePictureSet.NumPocLtCurr = j;
+		ReferencePictureSet.NumPocLtFoll = k;
+
+		auto const& PicXInDPB = [&PinnedSPS, &CurrentSlice, MaxPicOrderCntLsb, PicOrderCntVal = ReferencePictureSet.PicOrderCntVal, PicLayerId = ReferencePictureSet.PicLayerId](int32 PicOrderCount, bool DeltaPocMsbPresentFlag) -> uint8
+		{
+			if (DeltaPocMsbPresentFlag)
+			{
+				for (uint8 i = 0; i < 16 && i < PinnedSPS->MaxDpbSize; i++)
+				{
+					if (PicOrderCntVal[i] == PicOrderCount && CurrentSlice->nuh_layer_id == PicLayerId[i])
+					{
+						return i;
+					}
+				}
+			}
+			else
+			{
+				for (uint8 i = 0; i < 16 && i < PinnedSPS->MaxDpbSize; i++)
+				{
+					if ((PicOrderCntVal[i] & (FMath::Max<uint32>(MaxPicOrderCntLsb, 1) - 1)) == PicOrderCount && CurrentSlice->nuh_layer_id == PicLayerId[i])
+					{
+						return i;
+					}
+				}
+			}
+		
+			return 255; // NOTE (aidan) placholder for "no reference picture"
+		};
+
+		// Clear reference usage flags ready to be set for current frame
+		FMemory::Memset(ReferencePictureSet.PicUsage, 0, sizeof(ReferencePictureSet.PicUsage));
+
+		// (8-6) Generate Long-term Reference Pictures
+		for( i = 0; i < ReferencePictureSet.NumPocLtCurr; i++ )
+		{
+			ReferencePictureSet.RefPicSetLtCurr[i] = PicXInDPB(ReferencePictureSet.PocLtCurr[i], CurrDeltaPocMsbPresentFlag[i] > 0);
+			if(ReferencePictureSet.RefPicSetLtCurr[i] != 255)
+			{
+				ReferencePictureSet.PicUsage[i] |= EPictureUsage::LONG_TERM;
+			}
+		}
+	
+		for( i = 0; i < ReferencePictureSet.NumPocLtFoll; i++ )
+		{
+			ReferencePictureSet.RefPicSetLtFoll[i] = PicXInDPB(ReferencePictureSet.PocLtFoll[i], FollDeltaPocMsbPresentFlag[i] > 0);
+			if(ReferencePictureSet.RefPicSetLtFoll[i] != 255)
+			{
+				ReferencePictureSet.PicUsage[i] |= EPictureUsage::LONG_TERM;
+			}
+		}
+
+		// (8-7) Generate short-term reference pictures
+		for( i = 0; i < ReferencePictureSet.NumPocStCurrBefore; i++ )
+		{
+			ReferencePictureSet.RefPicSetStCurrBefore[i] = PicXInDPB(ReferencePictureSet.PocStCurrBefore[i], true);
+			if(ReferencePictureSet.RefPicSetStCurrBefore[i] != 255)
+			{
+				ReferencePictureSet.PicUsage[i] |= EPictureUsage::SHORT_TERM;
+			}
+		}
+	
+		for( i = 0; i < ReferencePictureSet.NumPocStCurrAfter; i++ )
+		{
+			ReferencePictureSet.RefPicSetStCurrAfter[i] = PicXInDPB(ReferencePictureSet.PocStCurrAfter[i], true);
+			if(ReferencePictureSet.RefPicSetStCurrAfter[i] != 255)
+			{
+				ReferencePictureSet.PicUsage[i] |= EPictureUsage::SHORT_TERM;
+			}
+		}
+	
+		for( i = 0; i < ReferencePictureSet.NumPocStFoll; i++ )
+		{
+			ReferencePictureSet.RefPicSetStFoll[i] = PicXInDPB(ReferencePictureSet.PocStFoll[i], 1);
+			if(ReferencePictureSet.RefPicSetStFoll[i] != 255)
+			{
+				ReferencePictureSet.PicUsage[i] |= EPictureUsage::SHORT_TERM;
+			}
+		}
+
+		// TODO (aidan) Unavailable frames estimation. Section 8.3.3.1 talks about image generation but not sure how to kick that off
+
+		// (8-8) Store old Ref Picture lists
+		const uint32 NumRpsCurrTempList0 = FMath::Max<uint32>(CurrentSlice->num_ref_idx_l0_active_minus1.Value + 1, ReferencePictureSet.NumPicTotalCurr);
+	
+		TArray<uint32> RefPicListTemp0;
+		RefPicListTemp0.SetNumZeroed(NumRpsCurrTempList0);
+	
+		uint32 rIdx;
+		for (rIdx = 0; rIdx < NumRpsCurrTempList0; rIdx++)
+		{
+			for (i = 0; i < ReferencePictureSet.NumPocStCurrBefore && rIdx < NumRpsCurrTempList0; rIdx++, i++)
+			{
+				RefPicListTemp0[rIdx] = ReferencePictureSet.RefPicSetStCurrBefore[i];
+			}
+			for (i = 0; i < ReferencePictureSet.NumPocStCurrAfter && rIdx < NumRpsCurrTempList0; rIdx++, i++)
+			{
+				RefPicListTemp0[rIdx] = ReferencePictureSet.RefPicSetStCurrAfter[i];
+			}
+			for (i = 0; i < ReferencePictureSet.NumPocLtCurr && rIdx < NumRpsCurrTempList0; rIdx++, i++)
+			{
+				RefPicListTemp0[rIdx] = ReferencePictureSet.RefPicSetLtCurr[i];
+			}
+			if (PinnedPPS->pps_scc_extension.pps_curr_pic_ref_enabled_flag)
+			{
+				RefPicListTemp0[rIdx++] = CurrPicIdx;
+			}
+		}
+
+		// (8-9) Construct RefPicList0
+		for (rIdx = 0; rIdx <= CurrentSlice->num_ref_idx_l0_active_minus1; rIdx++)
+		{
+			ReferencePictureSet.RefPicList0[rIdx] = CurrentSlice->ref_pic_list_modification.ref_pic_list_modification_flag_l0 ? RefPicListTemp0[CurrentSlice->ref_pic_list_modification.list_entry_l0[rIdx]] : RefPicListTemp0[rIdx];
+		}
+
+		if (PinnedPPS->pps_scc_extension.pps_curr_pic_ref_enabled_flag && !CurrentSlice->ref_pic_list_modification.ref_pic_list_modification_flag_l0 && NumRpsCurrTempList0 > (CurrentSlice->num_ref_idx_l0_active_minus1 + 1))
+		{
+			ReferencePictureSet.RefPicList0[CurrentSlice->num_ref_idx_l0_active_minus1] = CurrPicIdx;
+		}
+
+		// (8-10) Repeat for RefPicList1 when decoding a B slice
+		if (CurrentSlice->slice_type == EH265SliceType::B)
+		{
+			const uint32 NumRpsCurrTempList1 = FMath::Max<uint32>(CurrentSlice->num_ref_idx_l1_active_minus1 + 1, ReferencePictureSet.NumPicTotalCurr);
+
+			rIdx = 0;
+			TArray<uint32> RefPicListTemp1;
+			RefPicListTemp1.SetNumZeroed(NumRpsCurrTempList1);
+			while (rIdx < NumRpsCurrTempList1)
+			{
+				for (i = 0; i < ReferencePictureSet.NumPocStCurrAfter && rIdx < NumRpsCurrTempList1; rIdx++, i++)
+				{
+					RefPicListTemp1[rIdx] = ReferencePictureSet.RefPicSetStCurrAfter[i];
+				}
+				for (i = 0; i < ReferencePictureSet.NumPocStCurrBefore && rIdx < NumRpsCurrTempList1; rIdx++, i++)
+				{
+					RefPicListTemp1[rIdx] = ReferencePictureSet.RefPicSetStCurrBefore[i];
+				}
+				for (i = 0; i < ReferencePictureSet.NumPocLtCurr && rIdx < NumRpsCurrTempList1; rIdx++, i++)
+				{
+					RefPicListTemp1[rIdx] = ReferencePictureSet.RefPicSetLtCurr[i];
+				}
+				if (PinnedPPS->pps_scc_extension.pps_curr_pic_ref_enabled_flag)
+				{
+					RefPicListTemp1[rIdx++] = CurrPicIdx;
+				}
+			}
+
+			for (rIdx = 0; rIdx <= CurrentSlice->num_ref_idx_l0_active_minus1; rIdx++)
+			{
+				ReferencePictureSet.RefPicList1[rIdx] = CurrentSlice->ref_pic_list_modification.ref_pic_list_modification_flag_l0 ? RefPicListTemp1[CurrentSlice->ref_pic_list_modification.list_entry_l1[rIdx]] : RefPicListTemp1[rIdx];
+			}
+		}
+
+		LastRpsIdx = CurrentSlice->CurrPicOrderCntVal % PinnedSPS->MaxDpbSize;
+		LastPicOrderCntValue = CurrentSlice->CurrPicOrderCntVal;
+	}
+
+	return EAVResult::Success;
+}
diff --git a/Engine/Plugins/Experimental/AVCodecs/AVCodecsCore/Source/AVCodecsCore/Public/Audio/Resources/AudioResourceCPU.h b/Engine/Plugins/Experimental/AVCodecs/AVCodecsCore/Source/AVCodecsCore/Public/Audio/Resources/AudioResourceCPU.h
new file mode 100644
--- /dev/null
+++ b/Engine/Plugins/Experimental/AVCodecs/AVCodecsCore/Source/AVCodecsCore/Public/Audio/Resources/AudioResourceCPU.h
@@ -0,0 +1,33 @@
+// Copyright Epic Games, Inc. All Rights Reserved.
+
+#pragma once
+
+#include "AVContext.h"
+#include "Audio/AudioResource.h"
+
+/**
+ * CPU audio context and resource.
+ */
+
+class AVCODECSCORE_API FAudioContextCPU : public FAVContext
+{
+public:
+	FAudioContextCPU();
+};
+
+class AVCODECSCORE_API FAudioResourceCPU : public TAudioResource<FAudioContextCPU>
+{
+private:
+	TSharedPtr<float> Raw;
+
+public:
+	FORCEINLINE TSharedPtr<float> const& GetRaw() const { return Raw; }
+
+	FAudioResourceCPU(TSharedRef<FAVDevice> const& Device, TSharedPtr<float> const& Raw, FAVLayout const& Layout, FAudioDescriptor const& Descriptor);
+	virtual ~FAudioResourceCPU() override = default;
+
+	virtual FAVResult Validate() const override;
+};
+
+DECLARE_TYPEID(FAudioContextCPU, AVCODECSCORE_API);
+DECLARE_TYPEID(FAudioResourceCPU, AVCODECSCORE_API);
diff --git a/Engine/Plugins/Experimental/AVCodecs/AVCodecsCore/Source/AVCodecsCore/Public/Video/CodecUtils/CodecUtilsH264.h b/Engine/Plugins/Experimental/AVCodecs/AVCodecsCore/Source/AVCodecsCore/Public/Video/CodecUtils/CodecUtilsH264.h
new file mode 100644
--- /dev/null
+++ b/Engine/Plugins/Experimental/AVCodecs/AVCodecsCore/Source/AVCodecsCore/Public/Video/CodecUtils/CodecUtilsH264.h
@@ -0,0 +1,532 @@
+// Copyright Epic Games, Inc. All Rights Reserved.
+
+#pragma once
+
+#include "CoreMinimal.h"
+#include "Utils/BitstreamReader.h"
+
+#include "AVResult.h"
+#include "Video/VideoPacket.h"
+
+struct FH264ProfileDefinition;
+
+// Easy to understand user facing profiles from the T-Rec-H.264-201304 spec
+// there is another enum in CodecUtilsH264 that aligns the profiles with PIC and constraints
+// however users are not expected to understand those ones.
+enum class EH264Profile : uint8
+{
+	Auto,
+	CALVLC444Intra,
+	Baseline,
+	ConstrainedBaseline,
+	Main,
+	ScalableBaseline,
+	ScalableConstrainedBaseline,
+	ScalableHigh,
+	ScalableConstrainedHigh,
+	ScalableHighIntra,
+	Extended,
+	High,
+	ProgressiveHigh,
+	ConstrainedHigh,
+	High10,
+	High10Intra,
+	MultiviewHigh,
+	High422,
+	High422Intra,
+	StereoHigh,
+	MultiviewDepthHigh,
+	High444,
+	High444Intra,
+	MAX
+};
+
+extern AVCODECSCORE_API FH264ProfileDefinition GH264ProfileDefinitions[static_cast<uint8>(EH264Profile::MAX)];
+
+namespace UE::AVCodecCore::H264
+{
+	enum class ENaluType : uint8
+	{
+		Unspecified = 0,
+		SliceOfNonIdrPicture = 1,
+		SliceDataPartitionA = 2,
+		SliceDataPartitionB = 3,
+		SliceDataPartitionC = 4,
+		SliceIdrPicture = 5,
+		SupplementalEnhancementInformation = 6,
+		SequenceParameterSet = 7,
+		PictureParameterSet = 8,
+		AccessUnitDelimiter = 9,
+		EndOfSequence = 10,
+		EndOfStream = 11,
+		FillerData = 12,
+		SequenceParameterSetExtension = 13,
+		PrefixNalUnit = 14,
+		SubsetSequenceParameterSet = 15,
+		Reserved16 = 16,
+		Reserved17 = 17,
+		Reserved18 = 18,
+		SliceOfAnAuxiliaryCoded = 19,
+		SliceExtension = 20,
+		SliceExtensionForDepthView = 21,
+		Reserved22 = 22,
+		Reserved23 = 23,
+		Unspecified24 = 24,
+		Unspecified25 = 25,
+		Unspecified26 = 26,
+		Unspecified27 = 27,
+		Unspecified28 = 28,
+		Unspecified29 = 29,
+		Unspecified30 = 30,
+		Unspecified31 = 31
+	};
+	
+	enum class EH264ProfileIDC : uint8
+	{
+		Auto = 0,
+		CALVLC444Intra = 44,
+		Baseline = 66,
+		// ConstrainedBaseline = 66,			with constraint flag 1 set
+		Main = 77,
+		ScalableBaseline = 83,
+		// ScalableConstrainedBaseline = 83,	with constraint flag 5 set
+		ScalableHigh = 86,
+		// ScalableConstrainedHigh = 86,		with constraint flag 5 set
+		// ScalableHighIntra = 86,				with constraint flag 3 set
+		Extended = 88,
+		High = 100,
+		// ProgressiveHigh = 100,				with constraint flag 4 set
+		// ConstrainedHigh = 100,				with constraint flag 4 & 5 set
+		High10 = 110,
+		// High10Intra = 110,					with constraint flag 3 set
+		MultiviewHigh = 118,
+		High422 = 122,
+		// High422Intra = 122,					with constraint flag 3 set
+		StereoHigh = 128,
+		MultiviewDepthHigh = 138,
+		High444 = 244,
+		// High444Intra = 244,					with constraint flag 3 set
+	};
+
+	enum class EH264ConstraintFlag : uint8
+	{
+		None = 0,
+		Set0 = 1u << 0,
+		Set1 = 1u << 1,
+		Set2 = 1u << 2,
+		Set3 = 1u << 3,
+		Set4 = 1u << 4,
+		Set5 = 1u << 5
+	};
+
+	inline EH264ConstraintFlag operator|(EH264ConstraintFlag const& lhs, EH264ConstraintFlag const& rhs)
+	{
+		return (EH264ConstraintFlag)(static_cast<uint8>(lhs) | static_cast<uint8>(rhs));
+	}
+
+	inline EH264ConstraintFlag operator&(EH264ConstraintFlag const& lhs, EH264ConstraintFlag const& rhs)
+	{
+		return (EH264ConstraintFlag)(static_cast<uint8>(lhs) & static_cast<uint8>(rhs));
+	}
+
+	inline EH264ConstraintFlag operator|=(EH264ConstraintFlag const& lhs, EH264ConstraintFlag const& rhs)
+	{
+		return (EH264ConstraintFlag)(static_cast<uint8>(lhs) | static_cast<uint8>(rhs));
+	}
+
+	inline bool operator==(EH264ConstraintFlag lhs, EH264ConstraintFlag rhs)
+	{
+		return static_cast<uint8>(lhs) == (static_cast<uint8>(lhs) & static_cast<uint8>(rhs));
+	}
+
+	inline uint8 operator&(uint8 const& lhs, EH264ConstraintFlag const& rhs)
+	{
+		return lhs & static_cast<uint8>(rhs);
+	}
+
+	enum class EH264Level : uint8
+	{
+		Level_1b = 0,
+		Level_1 = 10,
+		Level_1_1 = 11,
+		Level_1_2 = 12,
+		Level_1_3 = 13,
+		Level_2 = 20,
+		Level_2_1 = 21,
+		Level_2_2 = 22,
+		Level_3 = 30,
+		Level_3_1 = 31,
+		Level_3_2 = 32,
+		Level_4 = 40,
+		Level_4_1 = 41,
+		Level_4_2 = 42,
+		Level_5 = 50,
+		Level_5_1 = 51,
+		Level_5_2 = 52,
+		// Level_6 = 60,
+		// Level_6_1 = 61,
+		// Level_6_2 = 62
+	};
+
+	enum class EH264AspectRatioIDC : uint8
+	{
+		Unspecified = 0,
+		Square = 1, // 1x1
+		_12x11 = 2,
+		_10x11 = 3,
+		_16x11 = 4,
+		_40x33 = 5,
+		_24x11 = 6,
+		_20x11 = 7,
+		_32x11 = 8,
+		_80x33 = 9,
+		_18x11 = 10,
+		_15x11 = 11,
+		_64x33 = 12,
+		_160x99 = 13,
+		_4x3 = 14,
+		_3x2 = 15,
+		_2x1 = 16,
+		Reserved = 17,
+		Extended_SAR = 255
+	};
+
+	enum class EH264VideoFormat : uint8
+	{
+		Component = 0,
+		PAL = 1,
+		NTSC = 2,
+		SECAM =3,
+		MAC = 4,
+		Unspecified = 5,
+		Reserved6 = 6,
+		Reserved7 = 7
+	};
+}
+
+struct FH264ProfileDefinition
+{
+	EH264Profile Profile;
+	UE::AVCodecCore::H264::EH264ProfileIDC PIDC;
+	UE::AVCodecCore::H264::EH264ConstraintFlag ConstraintFlags;
+	const TCHAR* Name;
+};
+
+namespace UE::AVCodecCore::H264
+{
+	inline int32 EBSPtoRBSP(uint8* OutBuf, const uint8* InBuf, int32 NumBytesIn)
+	{
+		uint8* OutBase = OutBuf;
+		while(NumBytesIn-- > 0)
+		{
+			uint8 b = *InBuf++;
+			*OutBuf++ = b;
+			if (b == 0)
+			{
+				if (NumBytesIn > 1)
+				{
+					if (InBuf[0] == 0x00 && InBuf[1] == 0x03)
+					{
+						*OutBuf++ = 0x00;
+						InBuf += 2;
+						NumBytesIn -= 2;
+					}
+				}
+			}
+		}
+		return OutBuf - OutBase;
+	}
+	
+	struct FNaluInfo
+	{
+		uint64 Start, Size;
+		uint8 StartCodeSize;
+		uint8 RefIdc;
+		ENaluType Type;
+		const uint8* Data;
+	};
+	
+	FAVResult FindNALUs(FVideoPacket const& InPacket, TArray<FNaluInfo>& FoundNalus);
+
+	constexpr uint8 Default_4x4_Intra[16] = {
+		6,13,13,20,
+		20,20,28,28,
+		28,28,32,32,
+		32,37,37,42
+	};
+	
+	constexpr uint8 Default_4x4_Inter[16] = {
+		10,14,14,20,
+		20,20,24,24,
+		24,24,27,27,
+		27,30,30,34
+	};
+
+	constexpr uint8 Default_8x8_Intra[64] = {
+		6,10,10,13,11,13,16,16,
+		16,16,18,18,18,18,18,23,
+		23,23,23,23,23,25,25,25,
+		25,25,25,25,27,27,27,27,
+		27,27,27,27,29,29,29,29,
+		29,29,29,31,31,31,31,31,
+		31,33,33,33,33,33,36,36,
+		36,36,38,38,38,40,40,42 };
+	
+	constexpr uint8 Default_8x8_Inter[64] = {
+		9,13,13,15,13,15,17,17,
+		17,17,19,19,19,19,19,21,
+		21,21,21,21,21,22,22,22,
+		22,22,22,22,24,24,24,24,
+		24,24,24,24,25,25,25,25,
+		25,25,25,27,27,27,27,27,
+		27,28,28,28,28,28,30,30,
+		30,30,32,32,32,33,33,35
+	};
+
+	constexpr uint8 ZigZag4x4[16] = {
+		0,1,5,6,
+		2,4,7,12,
+		3,8,11,13,
+		9,10,15,15
+	};
+
+	constexpr uint8 FieldScan4x4[16] = {
+		0,2,8,12,
+		1,5,9,13,
+		3,6,10,14,
+		4,7,11,15
+	};
+
+	constexpr uint8 ZigZag8x8[64] = {
+		0,1,5,6,14,15,27,28,
+		2,4,7,13,16,26,29,42,
+		3,8,12,17,25,30,41,43,
+		9,11,18,24,31,40,44,53,
+		10,19,23,32,39,45,52,54,
+		20,22,33,38,46,51,55,60,
+		21,34,37,47,50,56,59,61,
+		35,36,48,49,57,58,62,63
+	};
+
+	constexpr uint8 FieldScan8x8[64] = {
+		0,3,8,15,22,30,38,52,
+		1,4,14,21,29,37,45,53,
+		2,7,16,23,31,39,46,58,
+		5,9,20,28,36,44,51,59,
+		6,13,24,32,40,47,54,60,
+		10,17,25,33,41,48,55,61,
+		11,18,26,34,42,49,56,62,
+		12,19,27,35,43,50,57,63
+	};
+
+	inline void ScaleListToWeightScale(const bool bIsFrame, const uint8* ScalingList, const uint8 ListSize, uint8* OutWeightList /* called C in spec*/)
+	{
+		// OPTIMIZE (aidan) seems like this can be done faster but should be functional for now
+		for (uint8 i = 0; i < ListSize; i++)
+		{
+			if (ListSize == 16)
+			{
+				OutWeightList[i] = ScalingList[bIsFrame ? ZigZag4x4[i] : FieldScan4x4[i]];
+			}
+			else
+			{
+				OutWeightList[i] = ScalingList[bIsFrame ? ZigZag8x8[i] : FieldScan8x8[i]];
+			}
+		}
+	};
+	
+	struct SPS_t : public FNalu
+	{
+	public:
+		U<8, EH264ProfileIDC> profile_idc = EH264ProfileIDC::Auto;
+		U<8, EH264ConstraintFlag> constraint_flags = EH264ConstraintFlag::None;
+		U<8> level_idc; // + profile_idc + constraint_flags = EH264Level
+		UE seq_parameter_set_id; // 0...31 inclusive
+		UE chroma_format_idc = 1; // 0...3 inclusive
+		U<1> separate_colour_plane_flag; // bool
+		UE bit_depth_luma_minus8; // 0...6 inclusive
+		UE bit_depth_chroma_minus8; // 0...6 inclusive
+		U<1> qpprime_y_zero_transform_bypass_flag; // 0...1 inclusive
+		U<1> seq_scaling_matrix_present_flag; // bool
+		U<1> seq_scaling_list_present_flag[12]; // bool array
+		uint8 ScalingList4x4[6][16]; // -2^(1+BitDepth)...2^(1+BitDepth)-1 inclusive
+		uint8 ScalingList8x8[6][64]; // -2^(1+BitDepth)...2^(1+BitDepth)-1 inclusive
+		UE log2_max_frame_num_minus4; // 0...12 inclusive
+		UE pic_order_cnt_type; // 0...3 inclusive
+		UE log2_max_pic_order_cnt_lsb_minus4; // 0...12 inclusive
+		U<1> delta_pic_order_always_zero_flag; // bool 
+		SE offset_for_non_ref_pic; // −2^(31)+1...2^(31)−1 inclusive
+		SE offset_for_top_to_bottom_field; // −2^(31)+1...2^(31)−1 inclusive
+		UE num_ref_frames_in_pic_order_cnt_cycle; // 0...255 inclusive
+		SE offset_for_ref_frame[255]; // −2^(31)+1...2^(31)−1 inclusive
+		UE max_num_ref_frames; // 0...MaxDpbFrames inclusive MaxDpbFrames = Min( MaxDecPicBufSizeMacro / ( PicWidthMacro * FrameHeightMacro ), 16) (A.3.1 or A.3.2)
+		U<1> gaps_in_frame_num_value_allowed_flag; // bool
+		UE pic_width_in_mbs_minus1; // + 1 = Width in macroblocks
+		UE pic_height_in_map_units_minus1; // + 1 = Height 
+		U<1> frame_mbs_only_flag; // bool
+		U<1> mb_adaptive_frame_field_flag; // bool
+		U<1> direct_8x8_inference_flag; // bool
+		U<1> frame_cropping_flag; // bool 
+		UE frame_crop_left_offset;
+		UE frame_crop_right_offset;
+		UE frame_crop_top_offset;
+		UE frame_crop_bottom_offset;
+		U<1> vui_parameters_present_flag; // bool
+		// VUI Parameters
+		U<1> aspect_ratio_info_present_flag; // bool
+		U<8, EH264AspectRatioIDC> aspect_ratio_idc = EH264AspectRatioIDC::Unspecified;
+		U<16> sar_width; // Width in arbitrary unit
+		U<16> sar_height; // Height in same arbitrary unit as Width
+		U<1> overscan_info_present_flag; // bool
+		U<1> overscan_appropriate_flag; // bool
+		U<1> video_signal_type_present_flag; // bool
+		U<3, EH264VideoFormat> video_format = EH264VideoFormat::Unspecified;
+		U<1> video_full_range_flag;  // bool
+		U<1> colour_description_present_flag; // bool
+		U<8> colour_primaries = 2; // 0...255 inclusive but see Table E-3 as many values are Reserved
+		U<8> transfer_characteristics = 2; // 0...255 inclusive but see Table E-4 as many values are Reserved
+		U<8> matrix_coefficients = 2; // 0...255 inclusive but see Table E-5 as many values are Reserved
+		U<1> chroma_loc_info_present_flag; // bool
+		UE chroma_sample_loc_type_top_field; // 0...5 inclusive
+		UE chroma_sample_loc_type_bottom_field; // 0...5 inclusive
+		U<1> timing_info_present_flag; // bool
+		U<32> num_units_in_tick = 1; // > 0
+		U<32> time_scale = 1000; // defaulted to ms
+		U<1> fixed_frame_rate_flag; // bool
+		U<1> nal_hrd_parameters_present_flag; // bool
+		U<1> vcl_hrd_parameters_present_flag; // bool
+		U<1> low_delay_hrd_flag; // bool
+		U<1> pic_struct_present_flag; // bool
+		U<1> bitstream_restriction_flag; // bool
+		U<1> motion_vectors_over_pic_boundaries_flag; // bool
+		UE max_bytes_per_pic_denom;
+		UE max_bits_per_mb_denom;
+		UE log2_max_mv_length_horizontal;
+		UE log2_max_mv_length_vertical;
+		UE max_num_reorder_frames;
+		UE max_dec_frame_buffering;
+		// HRD Parameters
+		UE cpb_cnt_minus1; // 0...31 inclusive
+		U<4> bit_rate_scale;
+		U<4> cpb_size_scale;
+		UE bit_rate_value_minus1[31]; // 0...2^(32)−2 inclusive TODO (aidan) might need an infer stage to fill out based on profile
+		UE cpb_size_value_minus1[31]; // 0...2^(32)−2 inclusive TODO (aidan) might need an infer stage to fill out based on profile
+		U<1> cbr_flag[31]; // bool
+		U<5> initial_cpb_removal_delay_length_minus1 = 23; // 
+		U<5> cpb_removal_delay_length_minus1 = 23;
+		U<5> dpb_output_delay_length_minus1 = 23;
+		U<5> time_offset_length = 24;
+		
+		EH264Profile GetProfile() const
+		{
+			for (FH264ProfileDefinition const& H264Profile : GH264ProfileDefinitions)
+			{
+				if (profile_idc == H264Profile.PIDC && (constraint_flags == H264Profile.ConstraintFlags))
+				{
+					return H264Profile.Profile;
+				}
+			}
+
+			return EH264Profile::MAX;
+		}
+
+	};
+	
+	FAVResult ParseSPS(FBitstreamReader& Bitstream, FNaluInfo const& InNaluInfo, TMap<uint32, SPS_t>& OutMapSPS);
+
+	struct PPS_t : public FNalu
+	{
+		UE pic_parameter_set_id;
+		UE seq_parameter_set_id;
+		U<1> entropy_coding_mode_flag;
+		U<1> bottom_field_pic_order_in_frame_present_flag;
+		UE num_slice_groups_minus1;
+		UE slice_group_map_type;
+		TArray<UE> run_length_minus1; // hard to justify preallocating this as it could be PicWidthInMbs * PicHeightInMapUnits - 1
+		TArray<UE> top_left; // hard to justify preallocating this as it could be PicWidthInMbs * PicHeightInMapUnits - 1
+		TArray<UE> bottom_right; // hard to justify preallocating this as it could be PicWidthInMbs * PicHeightInMapUnits - 1
+		U<1> slice_group_change_direction_flag;
+		UE slice_group_change_rate_minus1;
+		UE pic_size_in_map_units_minus1;
+		TArray<U<>> slice_group_id; // hard to justify preallocating this as it could be pic_size_in_map_units_minus1 + 1
+		UE num_ref_idx_l0_default_active_minus1;
+		UE num_ref_idx_l1_default_active_minus1;
+		U<1> weighted_pred_flag;
+		U<2> weighted_bipred_idc;
+		SE pic_init_qp_minus26;
+		SE pic_init_qs_minus26;
+		SE chroma_qp_index_offset;
+		U<1> deblocking_filter_control_present_flag;
+		U<1> constrained_intra_pred_flag;
+		U<1> redundant_pic_cnt_present_flag;
+		U<1> transform_8x8_mode_flag;
+		U<1> pic_scaling_matrix_present_flag;
+		U<1> pic_scaling_list_present_flag[12];
+		uint8 ScalingList4x4[6][16]; // -2^(1+BitDepth)...2^(1+BitDepth)-1 inclusive
+		uint8 ScalingList8x8[6][64]; // -2^(1+BitDepth)...2^(1+BitDepth)-1 inclusive
+		SE second_chroma_qp_index_offset;
+	};
+	
+	FAVResult ParsePPS(FBitstreamReader& Bitstream, FNaluInfo const& InNaluInfo, TMap<uint32, SPS_t> const& InMapSPS, TMap<uint32, PPS_t>& OutMapPPS);
+
+	struct Slice_t : public FNalu
+	{
+		UE first_mb_in_slice;
+		UE slice_type;
+		UE pic_parameter_set_id;
+		U<2> colour_plane_id;
+		U<> frame_num;
+		U<1> field_pic_flag;
+		U<1> bottom_field_flag;
+		UE idr_pic_id;
+		U<> pic_order_cnt_lsb;
+		SE delta_pic_order_cnt_bottom;
+		SE delta_pic_order_cnt[2];
+		UE redundant_pic_cnt;
+		U<1> direct_spatial_mv_pred_flag;
+		U<1> num_ref_idx_active_override_flag;
+		UE num_ref_idx_l0_active_minus1;
+		UE num_ref_idx_l1_active_minus1;
+		UE cabac_init_idc;
+		SE slice_qp_delta;
+		U<1> sp_for_switch_flag;
+		SE slice_qs_delta;
+		UE disable_deblocking_filter_idc;
+		SE slice_alpha_c0_offset_div2;
+		SE slice_beta_offset_div2;
+		U<> slice_group_change_cycle;
+		
+		// ref_pic_list_modification
+		struct RefPic_t
+		{
+			U<1> bIsLongTerm;
+			UE pic_num;
+		};
+		
+		U<1> ref_pic_list_modification_flag_l0;
+		TArray<RefPic_t> RefPicList0;
+		U<1> ref_pic_list_modification_flag_l1;
+		TArray<RefPic_t> RefPicList1;
+
+		// pred_weight_table
+
+
+		
+		// dec_ref_pic_marking
+		U<1> no_output_of_prior_pic_flag;
+		U<1> long_term_reference_flag;
+
+		U<1> adaptive_ref_pic_marking_mode_flag;
+
+		
+	};
+	
+	FAVResult ParseSliceHeader(FBitstreamReader& Bitstream, FNaluInfo const& InNaluInfo, TMap<uint32, SPS_t> InMapSPS, TMap<uint32, PPS_t> const& InMapPPS, Slice_t& OutSlice);
+
+	// SEI can transmit arbitrary data so we probably want something less rigid for this
+	struct SEI_t
+	{
+	};
+	
+	FAVResult ParseSEI(FBitstreamReader& Bitstream, FNaluInfo const& InNaluInfo, SEI_t& OutSEI);
+} // namespace UE::AVCodecCore::H264
diff --git a/Engine/Plugins/Experimental/AVCodecs/AVCodecsCore/Source/AVCodecsCore/Public/Video/Decoders/Configs/VideoDecoderConfigH264.h b/Engine/Plugins/Experimental/AVCodecs/AVCodecsCore/Source/AVCodecsCore/Public/Video/Decoders/Configs/VideoDecoderConfigH264.h
new file mode 100644
--- /dev/null
+++ b/Engine/Plugins/Experimental/AVCodecs/AVCodecsCore/Source/AVCodecsCore/Public/Video/Decoders/Configs/VideoDecoderConfigH264.h
@@ -0,0 +1,25 @@
+// Copyright Epic Games, Inc. All Rights Reserved.
+
+#pragma once
+
+#include "Video/VideoDecoder.h"
+#include "Video/CodecUtils/CodecUtilsH264.h"
+
+/*
+ * Configuration settings for H264 decoders.
+ */
+struct AVCODECSCORE_API FVideoDecoderConfigH264 : public FVideoDecoderConfig
+{
+	TMap<uint32, UE::AVCodecCore::H264::SPS_t> SPS;
+	TMap<uint32, UE::AVCodecCore::H264::PPS_t> PPS;
+	TArray<UE::AVCodecCore::H264::SEI_t> SEI;
+
+	FVideoDecoderConfigH264(EAVPreset Preset = EAVPreset::Default)
+		: FVideoDecoderConfig(Preset)
+	{
+	}
+
+	FAVResult Parse(TSharedRef<FAVInstance> const& Instance, FVideoPacket const& Packet, TArray<UE::AVCodecCore::H264::Slice_t>& OutSlices);
+};
+
+DECLARE_TYPEID(FVideoDecoderConfigH264, AVCODECSCORE_API);
diff --git a/Engine/Plugins/Experimental/AVCodecs/AVCodecsCore/Source/AVCodecsCore/Public/Video/Decoders/Configs/VideoDecoderConfigH265.h b/Engine/Plugins/Experimental/AVCodecs/AVCodecsCore/Source/AVCodecsCore/Public/Video/Decoders/Configs/VideoDecoderConfigH265.h
new file mode 100644
--- /dev/null
+++ b/Engine/Plugins/Experimental/AVCodecs/AVCodecsCore/Source/AVCodecsCore/Public/Video/Decoders/Configs/VideoDecoderConfigH265.h
@@ -0,0 +1,115 @@
+// Copyright Epic Games, Inc. All Rights Reserved.
+
+#pragma once
+
+#include "Video/VideoDecoder.h"
+#include "Video/CodecUtils/CodecUtilsH265.h"
+
+enum class EPictureUsage : uint8
+{
+	UNUSED			= 0u,
+	SHORT_TERM		= 1u << 0,
+	LONG_TERM		= 1u << 1,		
+};
+
+ENUM_CLASS_FLAGS(EPictureUsage);
+
+/*
+ * Configuration settings for H265 decoders.
+ */
+struct FVideoDecoderConfigH265 : public FVideoDecoderConfig
+{
+	FVideoDecoderConfigH265(EAVPreset Preset = EAVPreset::Default)
+		: FVideoDecoderConfig(Preset)
+	{
+	}
+
+	AVCODECSCORE_API FAVResult Parse(TSharedRef<FAVInstance> const& Instance, FVideoPacket const& Packet, TArray<TSharedPtr<UE::AVCodecCore::H265::FNaluH265>>& Slices);
+
+	AVCODECSCORE_API FAVResult UpdateScalingLists(TSharedRef<UE::AVCodecCore::H265::FNaluSlice> CurrentSlice);
+
+	AVCODECSCORE_API FAVResult UpdateRPS(TSharedRef<UE::AVCodecCore::H265::FNaluSlice> CurrentSlice);
+
+	TMap<uint32, TSharedPtr<UE::AVCodecCore::H265::FNaluVPS>> ParsedVPS = {};
+	TMap<uint32, TSharedPtr<UE::AVCodecCore::H265::FNaluSPS>> ParsedSPS = {};
+	TMap<uint32, TSharedPtr<UE::AVCodecCore::H265::FNaluPPS>> ParsedPPS = {};
+	TArray<TSharedPtr<UE::AVCodecCore::H265::FNaluSEI>> ParsedSEI = {};
+
+	// Decoder state
+	int32 CurrPicIdx = -1;
+	int32 LastRpsIdx = 0;
+	int32 LastPicOrderCntValue = 0;
+
+	// Scaling lists
+	uint8 ScalingList4x4[6][16];
+	uint8 ScalingList8x8[6][16];
+	uint8 ScalingList16x16[6][16];
+	uint8 ScalingList32x32[6][16];
+
+	int16 ScalingListDCCoeff4x4[6];
+	int16 ScalingListDCCoeff8x8[6];
+	int16 ScalingListDCCoeff16x16[6];
+	int16 ScalingListDCCoeff32x32[6];
+	
+	struct FReferencePictureSet
+	{
+		uint8 NumPicTotalCurr;
+		
+		int32 PicOrderCntVal[16] = { 0 };
+		uint8 PicAvaliability[16] = { 0 };
+		EPictureUsage PicUsage[16]  = { EPictureUsage::UNUSED };
+		uint8 PicLayerId[16] = { 0 };
+		
+		uint8 NumPocStCurrBefore;
+		int32 PocStCurrBefore[16] = { 0 };
+		uint8 RefPicSetStCurrBefore[8] = { 0 };
+		
+		uint8 NumPocStCurrAfter;
+		int32 PocStCurrAfter[16] = { 0 };
+		uint8 RefPicSetStCurrAfter[8] = { 0 };
+
+		uint8 NumPocStFoll;
+		int32 PocStFoll[16] = { 0 };
+		uint8 RefPicSetStFoll[16] = { 0 };
+
+		uint8 NumPocLtCurr;
+		int32 PocLtCurr[16] = { 0 };
+		uint8 RefPicSetLtCurr[16] = { 0 };
+
+		uint8 NumPocLtFoll;
+		int32 PocLtFoll[16] = { 0 };
+		uint8 RefPicSetLtFoll[16] = { 0 };
+
+		uint8 RefPicList0[16] = { 0 };
+		uint8 RefPicList1[16] = { 0 };
+
+		uint8 RefPicSetInterLayer0[8] = { 0 };
+		uint8 RefPicSetInterLayer1[8] = { 0 };
+
+		void PrepState()
+		{
+			NumPocStCurrBefore = 0;
+			FMemory::Memset( PocStCurrBefore, 0, sizeof(PocStCurrBefore));
+			
+			NumPocStCurrAfter = 0;
+			FMemory::Memset( PocStCurrAfter, 0, sizeof(PocStCurrAfter));
+			
+			NumPocStFoll = 0;
+			FMemory::Memset( PocStFoll, 0, sizeof(PocStFoll));
+			
+			NumPocLtCurr = 0;
+			FMemory::Memset( PocLtCurr, 0, sizeof(PocLtCurr));
+			
+			NumPocLtFoll = 0;
+			FMemory::Memset( PocLtFoll, 0, sizeof(PocLtFoll));
+
+			NumPicTotalCurr = 0;
+
+			FMemory::Memset( PicUsage, 0, sizeof(PicUsage));
+
+			FMemory::Memset( PicOrderCntVal, 0, sizeof(PicOrderCntVal));
+		}
+	} ReferencePictureSet;
+};
+
+DECLARE_TYPEID(FVideoDecoderConfigH265, AVCODECSCORE_API);
diff --git a/Engine/Plugins/Experimental/AVCodecs/AVCodecsCore/Source/AVCodecsCore/Public/Video/Resources/VideoResourceCPU.h b/Engine/Plugins/Experimental/AVCodecs/AVCodecsCore/Source/AVCodecsCore/Public/Video/Resources/VideoResourceCPU.h
new file mode 100644
--- /dev/null
+++ b/Engine/Plugins/Experimental/AVCodecs/AVCodecsCore/Source/AVCodecsCore/Public/Video/Resources/VideoResourceCPU.h
@@ -0,0 +1,33 @@
+// Copyright Epic Games, Inc. All Rights Reserved.
+
+#pragma once
+
+#include "AVContext.h"
+#include "Video/VideoResource.h"
+
+/**
+ * CPU video context and resource.
+ */
+
+class AVCODECSCORE_API FVideoContextCPU : public FAVContext
+{
+public:
+	FVideoContextCPU();
+};
+
+class AVCODECSCORE_API FVideoResourceCPU : public TVideoResource<FVideoContextCPU>
+{
+private:
+	TSharedPtr<uint8> Raw;
+
+public:
+	FORCEINLINE TSharedPtr<uint8> const& GetRaw() const { return Raw; }
+
+	FVideoResourceCPU(TSharedRef<FAVDevice> const& Device, TSharedPtr<uint8> const& Raw, FAVLayout const& Layout, FVideoDescriptor const& Descriptor);
+	virtual ~FVideoResourceCPU() override = default;
+
+	virtual FAVResult Validate() const override;
+};
+
+DECLARE_TYPEID(FVideoContextCPU, AVCODECSCORE_API);
+DECLARE_TYPEID(FVideoResourceCPU, AVCODECSCORE_API);
diff --git a/Engine/Plugins/Experimental/AVCodecs/AVCodecsCore/Source/AVCodecsCore/Public/Video/Resources/VideoResourceVulkan.h b/Engine/Plugins/Experimental/AVCodecs/AVCodecsCore/Source/AVCodecsCore/Public/Video/Resources/VideoResourceVulkan.h
new file mode 100644
--- /dev/null
+++ b/Engine/Plugins/Experimental/AVCodecs/AVCodecsCore/Source/AVCodecsCore/Public/Video/Resources/VideoResourceVulkan.h
@@ -0,0 +1,63 @@
+// Copyright Epic Games, Inc. All Rights Reserved.
+
+#pragma once
+
+#include "AVContext.h"
+#include "Video/VideoResource.h"
+
+THIRD_PARTY_INCLUDES_START
+#if PLATFORM_WINDOWS
+#include "Windows/AllowWindowsPlatformTypes.h"
+#include "Windows/PreWindowsApi.h"
+#endif
+
+#include "vulkan_core.h"
+
+#if PLATFORM_WINDOWS
+#include "Windows/PostWindowsApi.h"
+#include "Windows/HideWindowsPlatformTypes.h"
+#endif
+THIRD_PARTY_INCLUDES_END
+
+/**
+ * Vulkan platform video context and resource.
+ */
+
+class AVCODECSCORE_API FVideoContextVulkan : public FAVContext
+{
+public:
+	VkInstance Instance;
+	VkDevice Device;
+	VkPhysicalDevice PhysicalDevice;
+	TFunction<PFN_vkVoidFunction(const char*)> vkGetDeviceProcAddr;
+
+	FVideoContextVulkan(VkInstance const& Instance, VkDevice const& Device, VkPhysicalDevice const& PhysicalDevice, TFunction<PFN_vkVoidFunction(const char*)> const& GetDeviceProcAddrFunc);
+};
+
+#if PLATFORM_WINDOWS
+typedef HANDLE FVulkanSharedHandle;
+#else
+typedef int FVulkanSharedHandle;
+#endif
+
+class AVCODECSCORE_API FVideoResourceVulkan : public TVideoResource<FVideoContextVulkan>
+{
+private:
+	VkDeviceMemory DeviceMemory;
+	FVulkanSharedHandle DeviceMemorySharedHandle = FVulkanSharedHandle();
+	VkFence Fence;
+	FVulkanSharedHandle FenceSharedHandle = FVulkanSharedHandle();
+
+public:
+	FORCEINLINE VkDeviceMemory GetRaw() const { return DeviceMemory; }
+	FORCEINLINE FVulkanSharedHandle const& GetSharedHandle() const { return DeviceMemorySharedHandle; }
+	FORCEINLINE FVulkanSharedHandle const& GetFenceSharedHandle() const { return FenceSharedHandle; }
+
+	FVideoResourceVulkan(TSharedRef<FAVDevice> const& Device, VkDeviceMemory Raw, FAVLayout const& Layout, FVideoDescriptor const& Descriptor);
+	virtual ~FVideoResourceVulkan() override = default;
+
+	virtual FAVResult Validate() const override;
+};
+
+DECLARE_TYPEID(FVideoContextVulkan, AVCODECSCORE_API);
+DECLARE_TYPEID(FVideoResourceVulkan, AVCODECSCORE_API);
diff --git a/Engine/Plugins/Experimental/AVCodecs/AVCodecsCore/Source/AVCodecsCore/Public/Video/Resources/Windows/VideoResourceD3D.h b/Engine/Plugins/Experimental/AVCodecs/AVCodecsCore/Source/AVCodecsCore/Public/Video/Resources/Windows/VideoResourceD3D.h
new file mode 100644
--- /dev/null
+++ b/Engine/Plugins/Experimental/AVCodecs/AVCodecsCore/Source/AVCodecsCore/Public/Video/Resources/Windows/VideoResourceD3D.h
@@ -0,0 +1,111 @@
+// Copyright Epic Games, Inc. All Rights Reserved.
+
+#pragma once
+
+#include "Templates/RefCounting.h"
+
+#include "AVContext.h"
+#include "Video/VideoResource.h"
+
+#pragma warning(push)
+#pragma warning(disable: 4005)
+
+THIRD_PARTY_INCLUDES_START
+#include "Windows/AllowWindowsPlatformTypes.h"
+#include "Windows/PreWindowsApi.h"
+#include <d3d11.h>
+#include <mfobjects.h>
+#include <mftransform.h>
+#include <mfapi.h>
+#include <mferror.h>
+#include <mfidl.h>
+#include <codecapi.h>
+#include <shlwapi.h>
+#include <mfreadwrite.h>
+#include <d3d11_1.h>
+#include <d3d12.h>
+#include <dxgi1_4.h>
+#include "Windows/PostWindowsApi.h"
+#include "Windows/HideWindowsPlatformTypes.h"
+THIRD_PARTY_INCLUDES_END
+
+/**
+ * D3D11 and D3D12 platform video context and resource.
+ */
+
+class AVCODECSCORE_API FVideoContextD3D11 : public FAVContext
+{
+public:
+	TRefCountPtr<ID3D11Device> Device;
+
+	FVideoContextD3D11(TRefCountPtr<ID3D11Device> const& Device);
+};
+
+class AVCODECSCORE_API FVideoContextD3D12 : public FAVContext
+{
+public:
+	TRefCountPtr<ID3D12Device> Device;
+
+	FVideoContextD3D12(TRefCountPtr<ID3D12Device> const& Device);
+};
+
+class AVCODECSCORE_API FVideoResourceD3D11 : public TVideoResource<FVideoContextD3D11>
+{
+private:
+	TRefCountPtr<ID3D11Texture2D> Raw;
+	HANDLE SharedHandle = nullptr;
+
+public:
+	static FVideoDescriptor GetDescriptorFrom(TSharedRef<FAVDevice> const& Device, TRefCountPtr<ID3D11Texture2D> const& Raw);
+
+	FORCEINLINE TRefCountPtr<ID3D11Texture2D> const& GetRaw() const { return Raw; }
+	FORCEINLINE HANDLE const& GetSharedHandle() const { return SharedHandle; }
+
+	FVideoResourceD3D11(TSharedRef<FAVDevice> const& Device, TRefCountPtr<ID3D11Texture2D> const& Raw, FAVLayout const& Layout);
+
+	virtual FAVResult Validate() const override;
+};
+
+class AVCODECSCORE_API FVideoResourceD3D12 : public TVideoResource<FVideoContextD3D12>
+{
+public:
+	struct FRawD3D12
+	{
+		TRefCountPtr<ID3D12Resource> D3DResource;
+		HANDLE D3DResourceShared = nullptr;
+		TRefCountPtr<ID3D12Heap> D3DHeap;
+		HANDLE D3DHeapShared = nullptr;
+		TRefCountPtr<ID3D12Fence> D3DFence;
+		HANDLE D3DFenceShared = nullptr;
+		uint64 FenceValue;
+	};
+
+private:
+	FRawD3D12 Raw;
+
+public:
+	static FVideoDescriptor GetDescriptorFrom(TSharedRef<FAVDevice> const& Device, FRawD3D12 const& Raw);
+
+	FORCEINLINE TRefCountPtr<ID3D12Resource> const& GetResource() const { return Raw.D3DResource; }
+	FORCEINLINE HANDLE const& GetResourceSharedHandle() const { return Raw.D3DResourceShared; }
+	
+	FORCEINLINE TRefCountPtr<ID3D12Heap> const& GetHeap() const { return Raw.D3DHeap; }
+	FORCEINLINE HANDLE const& GetHeapSharedHandle() const { return Raw.D3DHeapShared; }
+
+	FORCEINLINE TRefCountPtr<ID3D12Fence> const& GetFence() const { return Raw.D3DFence; };
+	FORCEINLINE HANDLE const& GetFenceSharedHandle() const { return Raw.D3DFenceShared; }
+	FORCEINLINE uint64 const& GetFenceValue() const { return Raw.FenceValue; }
+
+	uint64 GetSizeInBytes();
+
+	FVideoResourceD3D12(TSharedRef<FAVDevice> const& Device, FRawD3D12 const& Raw, FAVLayout const& Layout);
+	FVideoResourceD3D12(TSharedRef<FAVDevice> const& Device, FRawD3D12 const& Raw, FAVLayout const& Layout, FVideoDescriptor const& Descriptor);
+	virtual ~FVideoResourceD3D12() override;
+
+	virtual FAVResult Validate() const override;
+};
+
+DECLARE_TYPEID(FVideoContextD3D11, AVCODECSCORE_API);
+DECLARE_TYPEID(FVideoContextD3D12, AVCODECSCORE_API);
+DECLARE_TYPEID(FVideoResourceD3D11, AVCODECSCORE_API);
+DECLARE_TYPEID(FVideoResourceD3D12, AVCODECSCORE_API);
\ No newline at end of file
diff --git a/Engine/Plugins/Experimental/AVCodecs/AVCodecsCore/Source/AVCodecsCoreRHI/Private/AVCodecsCoreRHIModule.cpp b/Engine/Plugins/Experimental/AVCodecs/AVCodecsCore/Source/AVCodecsCoreRHI/Private/AVCodecsCoreRHIModule.cpp
new file mode 100644
--- /dev/null
+++ b/Engine/Plugins/Experimental/AVCodecs/AVCodecsCore/Source/AVCodecsCoreRHI/Private/AVCodecsCoreRHIModule.cpp
@@ -0,0 +1,92 @@
+// Copyright Epic Games, Inc. All Rights Reserved.
+
+#include "Modules/ModuleManager.h"
+
+#include "IVulkanDynamicRHI.h"
+
+#include "Video/Resources/VideoResourceRHI.h"
+#include "Video/Resources/VideoResourceVulkan.h"
+#include "Video/Decoders/VideoDecoderRHI.h"
+#include "Video/Decoders/Configs/VideoDecoderConfigH264.h"
+#include "Video/Decoders/Configs/VideoDecoderConfigH265.h"
+#include "Video/Encoders/VideoEncoderRHI.h"
+#include "Video/Encoders/Configs/VideoEncoderConfigH264.h"
+#include "Video/Encoders/Configs/VideoEncoderConfigH265.h"
+
+#if PLATFORM_WINDOWS
+#include "ID3D11DynamicRHI.h"
+#include "ID3D12DynamicRHI.h"
+
+#include "Video/Resources/Windows/VideoResourceD3D.h"
+#endif
+
+class FAVCodecsCoreRHI : public IModuleInterface
+{
+public:
+	virtual void StartupModule() override
+	{
+		if (FApp::CanEverRender())
+		{
+#if PLATFORM_WINDOWS
+			TCHAR const* DynamicRHIModuleName = GetSelectedDynamicRHIModuleName(false);
+#elif PLATFORM_LINUX
+			TCHAR const* DynamicRHIModuleName = TEXT("VulkanRHI");
+#endif
+
+			if (FString("VulkanRHI") == FString(DynamicRHIModuleName))
+			{
+#if PLATFORM_WINDOWS
+				TArray<ANSICHAR const*> const ExtensionsToAdd { VK_KHR_EXTERNAL_MEMORY_EXTENSION_NAME, VK_KHR_EXTERNAL_MEMORY_WIN32_EXTENSION_NAME };
+#elif PLATFORM_LINUX
+				TArray<ANSICHAR const*> const ExtensionsToAdd { VK_KHR_EXTERNAL_MEMORY_EXTENSION_NAME, VK_KHR_EXTERNAL_MEMORY_FD_EXTENSION_NAME };
+#endif
+
+				IVulkanDynamicRHI::AddEnabledDeviceExtensionsAndLayers(ExtensionsToAdd, TArray<ANSICHAR const*>());
+			}
+
+			FCoreDelegates::OnPostEngineInit.AddLambda([]()
+			{
+				switch (GDynamicRHI->GetInterfaceType())
+				{
+				case ERHIInterfaceType::Vulkan:
+					FAVDevice::GetHardwareDevice()->SetContext<FVideoContextVulkan>(
+						MakeShared<FVideoContextVulkan>(
+							GetIVulkanDynamicRHI()->RHIGetVkInstance(),
+							GetIVulkanDynamicRHI()->RHIGetVkDevice(),
+							GetIVulkanDynamicRHI()->RHIGetVkPhysicalDevice(),
+							[] (char const* Name) -> PFN_vkVoidFunction {
+								return (PFN_vkVoidFunction)GetIVulkanDynamicRHI()->RHIGetVkDeviceProcAddr(Name);
+							}));
+				
+					break;
+#if PLATFORM_WINDOWS
+				case ERHIInterfaceType::D3D11:
+					FAVDevice::GetHardwareDevice()->SetContext<FVideoContextD3D11>(
+						MakeShared<FVideoContextD3D11>(
+							GetID3D11DynamicRHI()->RHIGetDevice()));
+				
+					break;
+				case ERHIInterfaceType::D3D12:
+					FAVDevice::GetHardwareDevice()->SetContext<FVideoContextD3D12>(
+						MakeShared<FVideoContextD3D12>(
+							GetID3D12DynamicRHI()->RHIGetDevice(0)));
+				
+					break;
+#endif
+				default:
+					break;
+				}
+
+				FAVDevice::GetHardwareDevice()->SetContext<FVideoContextRHI>(MakeShared<FVideoContextRHI>());
+			});
+		}
+
+		FVideoEncoder::Register<TVideoEncoderRHI<FVideoEncoderConfigH264>, FVideoResourceRHI, FVideoEncoderConfigH264>();
+		FVideoEncoder::Register<TVideoEncoderRHI<FVideoEncoderConfigH265>, FVideoResourceRHI, FVideoEncoderConfigH265>();
+
+		FVideoDecoder::Register<TVideoDecoderRHI<FVideoDecoderConfigH264>, FVideoResourceRHI, FVideoDecoderConfigH264>();
+		FVideoDecoder::Register<TVideoDecoderRHI<FVideoDecoderConfigH265>, FVideoResourceRHI, FVideoDecoderConfigH265>();
+	}
+};
+
+IMPLEMENT_MODULE(FAVCodecsCoreRHI, AVCodecsCoreRHI);
diff --git a/Engine/Plugins/Experimental/AVCodecs/AVCodecsCore/Source/AVCodecsCoreRHI/Public/Audio/Encoders/SimpleAudioEncoder.h b/Engine/Plugins/Experimental/AVCodecs/AVCodecsCore/Source/AVCodecsCoreRHI/Public/Audio/Encoders/SimpleAudioEncoder.h
new file mode 100644
--- /dev/null
+++ b/Engine/Plugins/Experimental/AVCodecs/AVCodecsCore/Source/AVCodecsCoreRHI/Public/Audio/Encoders/SimpleAudioEncoder.h
@@ -0,0 +1,117 @@
+// Copyright Epic Games, Inc. All Rights Reserved.
+
+#pragma once
+
+#include "SampleBuffer.h"
+
+#include "Audio/SimpleAudio.h"
+#include "Audio/AudioEncoder.h"
+
+#include "CoreMinimal.h"
+#include "Containers/Queue.h"
+#include "HAL/Runnable.h"
+
+#include "SimpleAudioEncoder.generated.h"
+
+USTRUCT(BlueprintType)
+struct FSimpleAudioEncoderConfig
+{
+	GENERATED_BODY()
+
+public:
+	UPROPERTY(EditAnywhere, BlueprintReadWrite, Category = "Audio")
+	int32 Bitrate;
+	
+	UPROPERTY(EditAnywhere, BlueprintReadWrite, Category = "Audio")
+	int32 Samplerate;
+	
+	UPROPERTY(EditAnywhere, BlueprintReadWrite, Category = "Audio")
+	int32 NumChannels;
+	
+	FSimpleAudioEncoderConfig(ESimpleAVPreset Preset = ESimpleAVPreset::Default)
+		: FSimpleAudioEncoderConfig(FAudioEncoderConfig(USimpleAudioHelper::ConvertPreset(Preset)))
+	{
+	}
+
+	FSimpleAudioEncoderConfig(FAudioEncoderConfig const& From)
+		: Bitrate(From.Bitrate)
+		, Samplerate(From.Samplerate)
+		, NumChannels(From.NumChannels)
+	{
+	}
+	
+	operator FAudioEncoderConfig() const
+    {
+		FAudioEncoderConfig Result;
+		Result.Bitrate = Bitrate;
+		Result.Samplerate = Samplerate;
+		Result.NumChannels = NumChannels;
+
+		return Result;
+    }
+};
+
+UCLASS(Blueprintable)
+class AVCODECSCORERHI_API USimpleAudioEncoder : public UObject, public FRunnable
+{
+	GENERATED_BODY()
+
+private:
+	TSharedPtr<TAudioEncoder<class FAudioResourceCPU>> Child;
+
+	struct FAsyncFrame
+	{
+		TSharedPtr<FAudioResourceCPU> Resource;
+		uint32 Timestamp;
+
+		FAsyncFrame() = default;
+		FAsyncFrame(TSharedPtr<FAudioResourceCPU> const& Resource, uint32 Timestamp)
+			: Resource(Resource)
+			, Timestamp(Timestamp)
+		{
+		}
+	};
+
+	FRunnableThread* AsyncThread = nullptr;
+	TQueue<FAsyncFrame> AsyncQueue;
+	TQueue<TSharedPtr<FAudioResourceCPU>> AsyncPool;
+
+	virtual uint32 Run() override;
+	virtual void Exit() override;
+
+public:
+	UFUNCTION(BlueprintCallable, Category = "Audio")
+	bool IsAsync() const;
+
+	UFUNCTION(BlueprintCallable, Category = "Audio")
+	bool IsOpen() const;
+
+	UFUNCTION(BlueprintCallable, Category = "Audio")
+	UPARAM(DisplayName = "Was Success") bool Open(ESimpleAudioCodec Codec, FSimpleAudioEncoderConfig Config, bool bAsynchronous);
+
+	UFUNCTION(BlueprintCallable, Category = "Audio")
+	void Close();
+	
+	UFUNCTION(BlueprintCallable, Category = "Audio", meta = (DisplayName = "Send Frame (PCM16)"))
+	UPARAM(DisplayName = "Was Success") bool SendFrameFloat(TArray<float> const& Resource, double Timestamp, int32 NumSamples, float SampleDuration);
+
+	bool SendFrame(Audio::TSampleBuffer<float> const& Resource, double Timestamp);
+	bool SendFrame(float const* ResourceData, double Timestamp, int32 NumSamples, float SampleDuration);
+
+	UFUNCTION(BlueprintCallable, Category = "Audio")
+	UPARAM(DisplayName = "Was Success") bool ReceivePacket(FSimpleAudioPacket& OutPacket);
+
+	UFUNCTION(BlueprintCallable, Category = "Audio")
+	void ReceivePackets(TArray<FSimpleAudioPacket>& OutPackets);
+
+	// Configuration
+public:
+	UFUNCTION(BlueprintCallable, Category = "Audio")
+	ESimpleAudioCodec GetCodec() const;
+
+	UFUNCTION(BlueprintCallable, Category = "Audio")
+	FSimpleAudioEncoderConfig GetConfig() const;
+	
+	UFUNCTION(BlueprintCallable, Category = "Audio")
+	void SetConfig(FSimpleAudioEncoderConfig NewConfig);
+};
diff --git a/Engine/Plugins/Experimental/AVCodecs/AVCodecsCore/Source/AVCodecsCoreRHI/Public/Audio/SimpleAudio.h b/Engine/Plugins/Experimental/AVCodecs/AVCodecsCore/Source/AVCodecsCoreRHI/Public/Audio/SimpleAudio.h
new file mode 100644
--- /dev/null
+++ b/Engine/Plugins/Experimental/AVCodecs/AVCodecsCore/Source/AVCodecsCoreRHI/Public/Audio/SimpleAudio.h
@@ -0,0 +1,34 @@
+// Copyright Epic Games, Inc. All Rights Reserved.
+
+#pragma once
+
+#include "AVConfig.h"
+#include "Audio/AudioPacket.h"
+
+#include "SimpleAV.h"
+
+#include "SimpleAudio.generated.h"
+
+UENUM(BlueprintType)
+enum class ESimpleAudioCodec : uint8
+{
+	AAC = 0,
+};
+
+USTRUCT(BlueprintType)
+struct FSimpleAudioPacket
+{
+	GENERATED_BODY()
+
+public:
+	FAudioPacket RawPacket;
+};
+
+UCLASS(Abstract)
+class AVCODECSCORERHI_API USimpleAudioHelper : public USimpleAVHelper
+{
+	GENERATED_BODY()
+
+public:
+	static ESimpleAudioCodec GuessCodec(TSharedRef<FAVInstance> const& Instance);
+};
diff --git a/Engine/Plugins/Experimental/AVCodecs/AVCodecsCore/Source/AVCodecsCoreRHI/Public/SimpleAV.h b/Engine/Plugins/Experimental/AVCodecs/AVCodecsCore/Source/AVCodecsCoreRHI/Public/SimpleAV.h
new file mode 100644
--- /dev/null
+++ b/Engine/Plugins/Experimental/AVCodecs/AVCodecsCore/Source/AVCodecsCoreRHI/Public/SimpleAV.h
@@ -0,0 +1,27 @@
+// Copyright Epic Games, Inc. All Rights Reserved.
+
+#pragma once
+
+#include "AVConfig.h"
+
+#include "SimpleAV.generated.h"
+
+UENUM(BlueprintType)
+enum class ESimpleAVPreset : uint8
+{
+	Default,
+	
+	UltraLowQuality,
+	LowQuality,
+	HighQuality,
+	Lossless,
+};
+
+UCLASS(Abstract)
+class AVCODECSCORERHI_API USimpleAVHelper : public UObject
+{
+	GENERATED_BODY()
+
+public:
+	static EAVPreset ConvertPreset(ESimpleAVPreset From);
+};
diff --git a/Engine/Plugins/Experimental/AVCodecs/AVCodecsCore/Source/AVCodecsCoreRHI/Public/Video/Decoders/SimpleVideoDecoder.h b/Engine/Plugins/Experimental/AVCodecs/AVCodecsCore/Source/AVCodecsCoreRHI/Public/Video/Decoders/SimpleVideoDecoder.h
new file mode 100644
--- /dev/null
+++ b/Engine/Plugins/Experimental/AVCodecs/AVCodecsCore/Source/AVCodecsCoreRHI/Public/Video/Decoders/SimpleVideoDecoder.h
@@ -0,0 +1,51 @@
+// Copyright Epic Games, Inc. All Rights Reserved.
+
+#pragma once
+
+#include "Video/SimpleVideo.h"
+#include "Video/VideoDecoder.h"
+
+#include "CoreMinimal.h"
+#include "Containers/Queue.h"
+#include "HAL/Runnable.h"
+
+#include "SimpleVideoDecoder.generated.h"
+
+UCLASS(Blueprintable)
+class AVCODECSCORERHI_API USimpleVideoDecoder : public UObject, public FRunnable
+{
+	GENERATED_BODY()
+
+private:
+	TSharedPtr<TVideoDecoder<class FVideoResourceRHI>> Child;
+
+	FRunnableThread* AsyncThread = nullptr;
+	TQueue<FSimpleVideoPacket> AsyncQueue;
+
+	virtual uint32 Run() override;
+	virtual void Exit() override;
+
+public:
+	UFUNCTION(BlueprintCallable, Category = "Video")
+	bool IsAsync() const;
+
+	UFUNCTION(BlueprintCallable, Category = "Video")
+	bool IsOpen() const;
+
+	UFUNCTION(BlueprintCallable, Category = "Video")
+	UPARAM(DisplayName = "Was Success") bool Open(ESimpleVideoCodec Codec, bool bAsynchronous);
+
+	UFUNCTION(BlueprintCallable, Category = "Video")
+	void Close();
+
+	UFUNCTION(BlueprintCallable, Category = "Video")
+	UPARAM(DisplayName = "Was Success") bool SendPacket(FSimpleVideoPacket const& Packet);
+
+	UFUNCTION(BlueprintCallable, Category = "Video")
+	UPARAM(DisplayName = "Was Success") bool ReceiveFrame(UTextureRenderTarget2D* Resource);
+
+	// Configuration
+public:
+	UFUNCTION(BlueprintCallable, Category = "Video")
+	ESimpleVideoCodec GetCodec() const;
+};
diff --git a/Engine/Plugins/Experimental/AVCodecs/AVCodecsCore/Source/AVCodecsCoreRHI/Public/Video/Resources/VideoResourceRHI.h b/Engine/Plugins/Experimental/AVCodecs/AVCodecsCore/Source/AVCodecsCoreRHI/Public/Video/Resources/VideoResourceRHI.h
new file mode 100644
--- /dev/null
+++ b/Engine/Plugins/Experimental/AVCodecs/AVCodecsCore/Source/AVCodecsCoreRHI/Public/Video/Resources/VideoResourceRHI.h
@@ -0,0 +1,97 @@
+// Copyright Epic Games, Inc. All Rights Reserved.
+
+#pragma once
+
+#include "AVContext.h"
+#include "AVExtension.h"
+#include "Video/VideoResource.h"
+#include "RHI.h"
+
+#if PLATFORM_WINDOWS
+THIRD_PARTY_INCLUDES_START
+#include "d3d12.h"
+THIRD_PARTY_INCLUDES_END
+#endif
+
+class AVCODECSCORERHI_API FVideoContextRHI : public FAVContext
+{
+};
+
+class AVCODECSCORERHI_API FVideoResourceRHI : public TVideoResource<FVideoContextRHI>
+{
+public:
+	struct FRawData
+	{
+		FTextureRHIRef Texture;
+#if PLATFORM_WINDOWS
+		TRefCountPtr<ID3D12Fence> Fence; // TODO replace with an RHI fence and convert in a transform function
+#else 
+		FGPUFenceRHIRef Fence;
+#endif
+		uint64 FenceValue;
+	};
+
+private:
+	FRawData Raw;
+
+public:
+	static FAVLayout GetLayoutFrom(TSharedRef<FAVDevice> const& Device, FTextureRHIRef const& Raw);
+	static FVideoDescriptor GetDescriptorFrom(TSharedRef<FAVDevice> const& Device, FTextureRHIRef const& Raw);
+
+	static TSharedPtr<FVideoResourceRHI> Create(TSharedPtr<FAVDevice> const& Device, FVideoDescriptor const& Descriptor, ETextureCreateFlags AdditionalFlags = ETextureCreateFlags::None, bool bIsSRGB = true);
+
+
+	FORCEINLINE FRawData const& GetRaw() const { return Raw; }
+
+	FVideoResourceRHI(TSharedRef<FAVDevice> const& Device, FRawData const& Raw, FVideoDescriptor const& OverrideDescriptor);
+	FVideoResourceRHI(TSharedRef<FAVDevice> const& Device, FRawData const& Raw);
+	virtual ~FVideoResourceRHI() override = default;
+
+	virtual FAVResult Validate() const override;
+
+	virtual void Lock() override;
+	virtual FScopeLock LockScope() override;
+
+	void CopyFrom(TArrayView64<uint8> const& From);
+	void CopyFrom(TArrayView<uint8> const& From);
+	void CopyFrom(TArray64<uint8> const& From);
+	void CopyFrom(TArray<uint8> const& From);
+	void CopyFrom(FTextureRHIRef const& From);
+
+	TSharedPtr<FVideoResourceRHI> TransformResource(FVideoDescriptor const& OutDescriptor);
+	void TransformResourceTo(FRHICommandListImmediate& RHICmdList, FTextureRHIRef Target, FVideoDescriptor const& OutDescriptor);
+};
+
+class AVCODECSCORERHI_API FResolvableVideoResourceRHI : public TResolvableVideoResource<FVideoResourceRHI>
+{
+public:
+	ETextureCreateFlags AdditionalFlags = ETextureCreateFlags::None;
+	
+protected:
+	virtual TSharedPtr<FVideoResourceRHI> TryResolve(TSharedPtr<FAVDevice> const& Device, FVideoDescriptor const& Descriptor) override;
+};
+
+/*
+* Handles pixel format or colorspace transformations
+*/
+template <>
+FAVResult FAVExtension::TransformResource(TSharedPtr<FVideoResourceRHI>& OutResource, TSharedPtr<class FVideoResourceRHI> const& InResource);
+
+template <>
+FAVResult FAVExtension::TransformResource(TSharedPtr<FVideoResourceRHI>& OutResource, TSharedPtr<class FVideoResourceCPU> const& InResource);
+
+#if PLATFORM_WINDOWS
+
+template <>
+FAVResult FAVExtension::TransformResource(TSharedPtr<class FVideoResourceD3D11>& OutResource, TSharedPtr<FVideoResourceRHI> const& InResource);
+
+template <>
+FAVResult FAVExtension::TransformResource(TSharedPtr<class FVideoResourceD3D12>& OutResource, TSharedPtr<FVideoResourceRHI> const& InResource);
+
+#endif
+
+template <>
+FAVResult FAVExtension::TransformResource(TSharedPtr<class FVideoResourceVulkan>& OutResource, TSharedPtr<FVideoResourceRHI> const& InResource);
+
+DECLARE_TYPEID(FVideoContextRHI, AVCODECSCORERHI_API);
+DECLARE_TYPEID(FVideoResourceRHI, AVCODECSCORERHI_API);
\ No newline at end of file
diff --git a/Engine/Plugins/Experimental/AVCodecs/AVCodecsCore/Source/AVCodecsCoreRHI/Public/Video/SimpleVideo.h b/Engine/Plugins/Experimental/AVCodecs/AVCodecsCore/Source/AVCodecsCoreRHI/Public/Video/SimpleVideo.h
new file mode 100644
--- /dev/null
+++ b/Engine/Plugins/Experimental/AVCodecs/AVCodecsCore/Source/AVCodecsCoreRHI/Public/Video/SimpleVideo.h
@@ -0,0 +1,38 @@
+// Copyright Epic Games, Inc. All Rights Reserved.
+
+#pragma once
+
+#include "AVConfig.h"
+#include "Video/VideoPacket.h"
+
+#include "SimpleAV.h"
+
+#include "SimpleVideo.generated.h"
+
+UENUM(BlueprintType)
+enum class ESimpleVideoCodec : uint8
+{
+	H264 = 0,
+	H265 = 1,
+};
+
+USTRUCT(BlueprintType)
+struct FSimpleVideoPacket
+{
+	GENERATED_BODY()
+
+public:
+	FVideoPacket RawPacket;
+};
+
+UCLASS(Abstract)
+class AVCODECSCORERHI_API USimpleVideoHelper : public USimpleAVHelper
+{
+	GENERATED_BODY()
+
+public:
+	static ESimpleVideoCodec GuessCodec(TSharedRef<FAVInstance> const& Instance);
+
+	UFUNCTION(BlueprintCallable, Category = "Video", meta = (WorldContext = "WorldContext"))
+	static void ShareRenderTarget2D(class UTextureRenderTarget2D* RenderTarget);
+};
diff --git a/Engine/Plugins/Experimental/AVCodecs/NVCodecs/Source/NVCodecs/Public/Video/Resources/VideoResourceCUDA.h b/Engine/Plugins/Experimental/AVCodecs/NVCodecs/Source/NVCodecs/Public/Video/Resources/VideoResourceCUDA.h
new file mode 100644
--- /dev/null
+++ b/Engine/Plugins/Experimental/AVCodecs/NVCodecs/Source/NVCodecs/Public/Video/Resources/VideoResourceCUDA.h
@@ -0,0 +1,105 @@
+// Copyright Epic Games, Inc. All Rights Reserved.
+
+#pragma once
+
+#include "AVContext.h"
+#include "AVExtension.h"
+#include "Video/VideoResource.h"
+
+#include "CudaModule.h"
+
+struct FCUDAContextScope
+{
+	FCUDAContextScope(CUcontext Context)
+	{
+		FCUDAModule::CUDA().cuCtxPushCurrent(Context);
+	}
+
+	~FCUDAContextScope()
+	{
+		FCUDAModule::CUDA().cuCtxPopCurrent(nullptr);
+	}
+};
+
+class NVCODECS_API FVideoContextCUDA : public FAVContext
+{
+public:
+	CUcontext Raw;
+
+	FVideoContextCUDA(CUcontext const& Raw);
+};
+
+class NVCODECS_API FVideoResourceCUDA : public TVideoResource<FVideoContextCUDA>
+{
+private:
+	CUexternalSemaphore ExternalSemaphore = nullptr;
+	uint64 SemaphoreValue = 0;
+	CUexternalMemory ExternalArray = nullptr;
+	CUmipmappedArray MipArray = nullptr;
+	CUarray MaxMipArray = nullptr;
+
+	CUarray Raw = nullptr;
+
+public:
+	FORCEINLINE CUarray GetRaw() const { return Raw; }
+
+	FVideoResourceCUDA(TSharedRef<FAVDevice> const& Device, CUarray Raw, FAVLayout const& Layout, FVideoDescriptor const& Descriptor);
+	FVideoResourceCUDA(TSharedRef<FAVDevice> const& Device, CUDA_EXTERNAL_MEMORY_HANDLE_DESC const& ExternalResourceDesc, CUDA_EXTERNAL_SEMAPHORE_HANDLE_DESC const& ExternalFenceDesc, FAVLayout const& Layout, FVideoDescriptor const& Descriptor, uint64 FenceValue = 0);
+	FVideoResourceCUDA(TSharedRef<FAVDevice> const& Device, CUDA_EXTERNAL_MEMORY_HANDLE_DESC const& ExternalResourceDesc, CUDA_EXTERNAL_SEMAPHORE_HANDLE_DESC const& ExternalFenceDesc, TSharedRef<FVideoResource> const& ExternalResource, uint64 FenceValue = 0);
+	virtual ~FVideoResourceCUDA() override;
+
+	virtual FAVResult Validate() const override;
+
+	FAVResult ReadData(TArray64<uint8>& OutData) const;
+
+	FAVResult CopyFrom(CUdeviceptr Target, uint32 MapPitch);
+	FAVResult CopyFromAsync(CUdeviceptr Target, uint32 MapPitch);
+
+	FAVResult CopyTo(CUdeviceptr Source, uint32 MapPitch);
+	FAVResult CopyToAsync(CUdeviceptr Source, uint32 MapPitch);
+
+private:
+	FORCEINLINE CUDA_MEMCPY2D CreateCopyFromParams(CUdeviceptr Source, uint32 SourcePitch)
+	{
+		CUDA_MEMCPY2D CopyParams = {};
+		CopyParams.srcMemoryType = CU_MEMORYTYPE_DEVICE;
+		CopyParams.srcDevice = Source;
+		CopyParams.srcPitch = SourcePitch;
+		CopyParams.dstMemoryType = CU_MEMORYTYPE_ARRAY;
+		CopyParams.dstArray = Raw;
+		CopyParams.WidthInBytes = GetRawDescriptor().Width * (GetDescriptor().Format == EVideoFormat::P010 ? 2 : 1);
+		CopyParams.Height = GetRawDescriptor().Height;
+
+		return CopyParams;
+	}
+
+	FORCEINLINE CUDA_MEMCPY2D CreateCopyToParams(CUdeviceptr Target, uint32 TargetPitch)
+	{
+		CUDA_MEMCPY2D CopyParams = {};
+		CopyParams.srcMemoryType = CU_MEMORYTYPE_ARRAY;
+		CopyParams.srcArray = Raw;
+		CopyParams.dstMemoryType = CU_MEMORYTYPE_DEVICE;
+		CopyParams.dstDevice = Target;
+		CopyParams.dstPitch = TargetPitch;
+		CopyParams.WidthInBytes = GetRawDescriptor().Width * (GetDescriptor().Format == EVideoFormat::P010 ? 2 : 1);
+		CopyParams.Height = GetRawDescriptor().Height;
+
+		return CopyParams;
+	}
+};
+
+#if PLATFORM_WINDOWS
+
+template <>
+FAVResult FAVExtension::TransformResource(TSharedPtr<FVideoResourceCUDA>& OutResource, TSharedPtr<class FVideoResourceD3D11> const& InResource);
+
+template <>
+FAVResult FAVExtension::TransformResource(TSharedPtr<FVideoResourceCUDA>& OutResource, TSharedPtr<class FVideoResourceD3D12> const& InResource);
+
+#endif
+
+template <>
+FAVResult FAVExtension::TransformResource(TSharedPtr<FVideoResourceCUDA>& OutResource, TSharedPtr<class FVideoResourceVulkan> const& InResource);
+
+DECLARE_TYPEID(FVideoContextCUDA, NVCODECS_API);
+DECLARE_TYPEID(FVideoResourceCUDA, NVCODECS_API);
\ No newline at end of file
diff --git a/Engine/Plugins/Experimental/AVCodecs/NVCodecs/Source/NVCodecsRHI/Private/NVCodecRHIModule.cpp b/Engine/Plugins/Experimental/AVCodecs/NVCodecs/Source/NVCodecsRHI/Private/NVCodecRHIModule.cpp
new file mode 100644
--- /dev/null
+++ b/Engine/Plugins/Experimental/AVCodecs/NVCodecs/Source/NVCodecsRHI/Private/NVCodecRHIModule.cpp
@@ -0,0 +1,27 @@
+// Copyright Epic Games, Inc. All Rights Reserved.
+
+#include "Modules/ModuleManager.h"
+#include "Misc/App.h"
+#include "Misc/CoreDelegates.h"
+#include "RHI.h"
+
+#include "NVDEC.h"
+#include "NVENC.h"
+
+class FNVCodecRHIModule : public IModuleInterface
+{
+public:
+	virtual void StartupModule() override
+	{
+		if (FApp::CanEverRender())
+		{
+			FCoreDelegates::OnPostEngineInit.AddLambda([]()
+				{
+					const_cast<FNVDEC&>(FAPI::Get<FNVDEC>()).bHasCompatibleGPU = IsRHIDeviceNVIDIA();
+					const_cast<FNVENC&>(FAPI::Get<FNVENC>()).bHasCompatibleGPU = IsRHIDeviceNVIDIA();
+				});
+		}
+	}
+};
+
+IMPLEMENT_MODULE(FNVCodecRHIModule, NVCodecsRHI);
diff --git a/Engine/Plugins/Experimental/AVCodecs/NVCodecs/Source/NVDEC/Private/NVDECModule.cpp b/Engine/Plugins/Experimental/AVCodecs/NVCodecs/Source/NVDEC/Private/NVDECModule.cpp
new file mode 100644
--- /dev/null
+++ b/Engine/Plugins/Experimental/AVCodecs/NVCodecs/Source/NVDEC/Private/NVDECModule.cpp
@@ -0,0 +1,34 @@
+// Copyright Epic Games, Inc. All Rights Reserved.
+
+#include "Video/Decoders/VideoDecoderNVDEC.h"
+#include "Video/Decoders/Configs/VideoDecoderConfigH264.h"
+#include "Video/Decoders/Configs/VideoDecoderConfigH265.h"
+#include "Video/Resources/VideoResourceVulkan.h"
+
+#if PLATFORM_WINDOWS
+#include "Video/Resources/Windows/VideoResourceD3D.h"
+#endif
+
+class FNVDECModule : public IModuleInterface
+{
+public:
+	virtual void StartupModule() override
+	{
+		FVideoDecoder
+			::RegisterPermutationsOf<FVideoDecoderNVDEC>
+			::With<
+				FVideoResourceCUDA,
+#if PLATFORM_WINDOWS
+				FVideoResourceD3D11, 
+				FVideoResourceD3D12, 
+#endif
+				FVideoResourceVulkan>
+			::And<FVideoDecoderConfigNVDEC, FVideoDecoderConfigH264, FVideoDecoderConfigH265>(
+				[](TSharedRef<FAVDevice> const& NewDevice, TSharedRef<FAVInstance> const& NewInstance)
+				{
+					return FAPI::Get<FNVDEC>().IsValid() && NewDevice->HasContext<FVideoContextCUDA>();
+				});
+	}
+};
+
+IMPLEMENT_MODULE(FNVDECModule, NVDEC);
diff --git a/Engine/Plugins/Experimental/AVCodecs/NVCodecs/Source/NVDEC/Private/Video/Decoders/Configs/VideoDecoderConfigNVDEC.cpp b/Engine/Plugins/Experimental/AVCodecs/NVCodecs/Source/NVDEC/Private/Video/Decoders/Configs/VideoDecoderConfigNVDEC.cpp
new file mode 100644
--- /dev/null
+++ b/Engine/Plugins/Experimental/AVCodecs/NVCodecs/Source/NVDEC/Private/Video/Decoders/Configs/VideoDecoderConfigNVDEC.cpp
@@ -0,0 +1,572 @@
+// Copyright Epic Games, Inc. All Rights Reserved.
+
+#include "Video/Decoders/Configs/VideoDecoderConfigNVDEC.h"
+
+REGISTER_TYPEID(FVideoDecoderConfigNVDEC);
+
+FAVResult FVideoDecoderConfigNVDEC::Parse(TSharedRef<FAVInstance> const& Instance, FVideoPacket const& Packet, TArray<FParsedPicture>& OutPictures)
+{
+	FAVResult Result = EAVResult::Success;
+
+	switch (CodecType)
+	{
+	case cudaVideoCodec_H264:
+		{
+			using namespace UE::AVCodecCore::H264;
+
+			FVideoDecoderConfigH264& H264 = Instance->Edit<FVideoDecoderConfigH264>();
+
+			TArray<Slice_t> Slices;
+			Result = H264.Parse(Instance, Packet, Slices);
+			if (Result.IsSuccess())
+			{
+				for (Slice_t const& Slice : Slices)
+				{
+					FParsedPicture& OutPicture = OutPictures.AddDefaulted_GetRef();
+					
+				
+					OutPicture.DecodeCreateInfo = *this;
+
+					CUVIDH264PICPARAMS& OutPictureH264 = OutPicture.CodecSpecific.h264;
+
+					// PPS
+					if (PPS_t const* const PPS = H264.PPS.Find(Slice.pic_parameter_set_id))
+					{
+						// SPS
+						if (SPS_t const* const SPS = H264.SPS.Find(PPS->seq_parameter_set_id))
+						{
+							OutPicture.DecodeCreateInfo.ulWidth = OutPicture.DecodeCreateInfo.ulTargetWidth = (SPS->pic_width_in_mbs_minus1 + 1) * 16;
+							OutPicture.DecodeCreateInfo.ulHeight = OutPicture.DecodeCreateInfo.ulTargetHeight = (SPS->pic_height_in_map_units_minus1 + 1) * 16;
+							OutPicture.DecodeCreateInfo.ChromaFormat = static_cast<cudaVideoChromaFormat>(SPS->chroma_format_idc.Value);
+							OutPicture.DecodeCreateInfo.ulIntraDecodeOnly = static_cast<unsigned long>(SPS->constraint_flags.Value) & EH264ConstraintFlag::Set3;
+						
+							OutPictureH264.entropy_coding_mode_flag = PPS->entropy_coding_mode_flag;
+							OutPictureH264.pic_order_present_flag = PPS->bottom_field_pic_order_in_frame_present_flag;
+							OutPictureH264.num_ref_idx_l0_active_minus1 = PPS->num_ref_idx_l0_default_active_minus1;
+							OutPictureH264.num_ref_idx_l1_active_minus1 = PPS->num_ref_idx_l1_default_active_minus1;
+							OutPictureH264.weighted_pred_flag = PPS->weighted_pred_flag;
+							OutPictureH264.weighted_bipred_idc = PPS->weighted_bipred_idc;
+							OutPictureH264.pic_init_qp_minus26 = PPS->pic_init_qp_minus26;
+							OutPictureH264.deblocking_filter_control_present_flag = PPS->deblocking_filter_control_present_flag;
+							OutPictureH264.redundant_pic_cnt_present_flag = PPS->redundant_pic_cnt_present_flag;
+							OutPictureH264.transform_8x8_mode_flag = PPS->transform_8x8_mode_flag;
+							OutPictureH264.MbaffFrameFlag = 0; // HACK does not seem to align directly with anything in spec
+							OutPictureH264.constrained_intra_pred_flag = PPS->constrained_intra_pred_flag;
+							OutPictureH264.chroma_qp_index_offset = PPS->chroma_qp_index_offset;
+							OutPictureH264.second_chroma_qp_index_offset = PPS->second_chroma_qp_index_offset;
+						
+							OutPictureH264.log2_max_frame_num_minus4 = SPS->log2_max_frame_num_minus4;
+							OutPictureH264.pic_order_cnt_type = SPS->pic_order_cnt_type;
+							OutPictureH264.log2_max_pic_order_cnt_lsb_minus4 = SPS->log2_max_pic_order_cnt_lsb_minus4;
+							OutPictureH264.delta_pic_order_always_zero_flag = SPS->delta_pic_order_always_zero_flag;
+							OutPictureH264.frame_mbs_only_flag = SPS->frame_mbs_only_flag;
+							OutPictureH264.direct_8x8_inference_flag = SPS->direct_8x8_inference_flag;
+							OutPictureH264.num_ref_frames = SPS->num_ref_frames_in_pic_order_cnt_cycle;
+							OutPictureH264.residual_colour_transform_flag = 0; // HACK does not seem to align directly with anything in spec
+							OutPictureH264.bit_depth_luma_minus8 = SPS->bit_depth_luma_minus8;
+							OutPictureH264.bit_depth_chroma_minus8 = SPS->bit_depth_chroma_minus8;
+							OutPictureH264.qpprime_y_zero_transform_bypass_flag = SPS->qpprime_y_zero_transform_bypass_flag;
+
+							// TODO (aidan) Other stuff? Not directly parsed from something must be inferred
+							OutPictureH264.ref_pic_flag = 0;
+							OutPictureH264.frame_num = 0;
+							OutPictureH264.CurrFieldOrderCnt[0] = 0;
+							OutPictureH264.CurrFieldOrderCnt[1] = 0;
+
+							// DPB
+							// TODO (aidan) this is complex as not directly parsed from something and must be inferred 
+							CUVIDH264DPBENTRY* DPB = OutPictureH264.dpb;
+							for (uint8 i = 0; i < 16u; i++)
+							{
+								// HACK (aidan)
+								static uint64 frameID = 0;
+					
+								DPB[0].PicIdx = frameID++;
+								DPB[0].used_for_reference = 3;
+
+								// TODO (aidan) handle case where we have a long term enabled  
+								DPB[0].is_long_term = 0;
+								DPB[0].FrameIdx = Slice.frame_num;
+
+								DPB[0].FieldOrderCnt[0] = 0;
+								DPB[0].FieldOrderCnt[1] = 0;
+							}
+
+							uint8 iYCbCr = SPS->separate_colour_plane_flag ? 1 : 0; // HACK (aidan) this comes from slice data that we are not yet parsing
+							bool bMbIsInterFlag = true;										   // HACK (aidan) part of the slice
+
+							if (PPS->pic_scaling_matrix_present_flag)
+							{
+								// Quantization Matrices
+								for (uint8 i = 0; i < 12u; i++)
+								{
+									// invoking 8.5.6 with ScalingList4x4[ iYCbCr + ( (mbIsInterFlag == 1 ) ? 3 : 0 )]
+									if (i < 6) // 4x4
+										{
+										ScaleListToWeightScale(true, PPS->ScalingList4x4[iYCbCr + (bMbIsInterFlag ? 3 : 0)], 16, OutPictureH264.WeightScale4x4[i]);
+										}
+									else // 8x8
+										{
+										ScaleListToWeightScale(true, PPS->ScalingList8x8[iYCbCr + (bMbIsInterFlag ? 3 : 0)], 64, OutPictureH264.WeightScale8x8[i]);
+										}
+								}
+							}
+							else if (SPS->seq_scaling_matrix_present_flag)
+							{
+								// Quantization Matrices
+								for (uint8 i = 0; i < 12u; i++)
+								{
+									// invoking 8.5.6 with ScalingList4x4[ iYCbCr + ( (mbIsInterFlag == 1 ) ? 3 : 0 )]
+									if (i < 6) // 4x4
+										{
+										ScaleListToWeightScale(true, SPS->ScalingList4x4[iYCbCr + (bMbIsInterFlag ? 3 : 0)], 16, OutPictureH264.WeightScale4x4[i]);
+										}
+									else // 8x8
+										{
+										ScaleListToWeightScale(true, SPS->ScalingList8x8[iYCbCr + (bMbIsInterFlag ? 3 : 0)], 64, OutPictureH264.WeightScale8x8[i]);
+										}
+								}
+							}
+						}
+						else
+						{
+							return FAVResult(EAVResult::ErrorInvalidState, FString::Printf(TEXT("SPS %d not found"), PPS->seq_parameter_set_id.Value), TEXT("H264"));
+						}
+					}
+					else
+					{
+						return FAVResult(EAVResult::ErrorInvalidState, FString::Printf(TEXT("PPS %d not found"), Slice.pic_parameter_set_id.Value), TEXT("H264"));
+					}
+				}
+			}
+		}
+
+		break;
+	case cudaVideoCodec_HEVC:
+		{
+			using namespace UE::AVCodecCore::H265;
+
+			FVideoDecoderConfigH265& H265 = Instance->Edit<FVideoDecoderConfigH265>();
+
+			TArray<TSharedPtr<FNaluH265>> Nalus;
+			Result = H265.Parse(Instance, Packet, Nalus);
+			if (Result.IsNotSuccess())
+			{
+				return FAVResult(EAVResult::ErrorResolving, TEXT("Failed to parse H265 Bitstream"));
+			}
+
+			// Loop through all the received NALUs and and add slices to the a FParsedPicture
+			// Depending on how we are receiving the bitstream we might get more than one picture
+			// so we handle that case by making multiple FParsedPictures and queue them
+			int32 index = 0;
+			while (index < Nalus.Num())
+			{
+				FParsedPicture& OutPicture = OutPictures.AddZeroed_GetRef();
+				
+				// NVDEC expects the bitstream to start with the first slice NAL delimiter at 0
+				OutPicture.pBitstreamData = Nalus[index]->StartCodeSize == 3 ? Nalus[index]->EBSP : Nalus[index]->EBSP + 1;
+
+				bool bParsedFirstSlice = false;
+
+				// TODO (aidan) this loop needs a clean-up wrt handling non-slices or non-first-segment slices
+				for (; index < Nalus.Num() && Nalus[index]->IsSlice(); index++)
+				{
+					TSharedPtr<FNaluSlice> Slice = StaticCastSharedPtr<FNaluSlice>(Nalus[index]);
+
+					// Apparently we got a new first slice without some other kind of
+					// NALU in between so we are done with the current picture
+					if (bParsedFirstSlice && Slice->first_slice_segment_in_pic_flag)
+					{
+						break;
+					}
+
+					// First slice in set
+					if (Slice->first_slice_segment_in_pic_flag == 1)
+					{
+						CUVIDHEVCPICPARAMS& OutPictureH265 = OutPicture.CodecSpecific.hevc;
+						
+						// PPS
+						if (!H265.ParsedPPS.Contains(Slice->slice_pic_parameter_set_id))
+						{
+							return FAVResult(EAVResult::Error, TEXT("PPS for slice missing from Instance"));
+						}
+
+						TSharedPtr<FNaluPPS> PPS = H265.ParsedPPS.FindRef(Slice->slice_pic_parameter_set_id);
+
+						// SPS
+						if (!H265.ParsedSPS.Contains(PPS->pps_seq_parameter_set_id))
+						{
+							return FAVResult(EAVResult::Error, TEXT("SPS for slice missing from Instance"));
+						}
+
+						TSharedPtr<FNaluSPS> SPS = H265.ParsedSPS.FindRef(PPS->pps_seq_parameter_set_id);
+
+						// VPS
+						if (!H265.ParsedVPS.Contains(SPS->sps_video_parameter_set_id))
+						{
+							return FAVResult(EAVResult::Error, TEXT("VPS for slice missing from Instance"));
+						}
+
+						TSharedPtr<FNaluVPS> VPS = H265.ParsedVPS.FindRef(SPS->sps_video_parameter_set_id);
+												
+						H265.UpdateRPS(Slice.ToSharedRef());
+												
+						OutPicture.DecodeCreateInfo = *this;
+
+						CodecType = OutPicture.DecodeCreateInfo.CodecType = cudaVideoCodec_HEVC;
+						
+						OutPicture.CurrPicIdx = H265.CurrPicIdx;
+						
+						OutPicture.field_pic_flag = 0;
+						OutPicture.ref_pic_flag = 1;
+						OutPicture.intra_pic_flag = Slice->IsIDR() || PPS->constrained_intra_pred_flag || true; // HACK (aidan.possemiers) We only currently send intra frames
+
+						if (OutPicture.intra_pic_flag) // HACK for intra only frames
+						{
+							FMemory::Memset(RefPicIdx, -1, sizeof(RefPicIdx));
+						}
+
+						OutPicture.PicWidthInMbs = SPS->pic_width_in_luma_samples / 16;
+						OutPicture.FrameHeightInMbs = SPS->pic_height_in_luma_samples / 16;
+
+						ulNumDecodeSurfaces = OutPicture.DecodeCreateInfo.ulNumDecodeSurfaces = 8; //SPS->MaxDpbSize; //HACK (aidan.possemiers) we should get this from the bitstream
+						ulNumOutputSurfaces = OutPicture.DecodeCreateInfo.ulNumOutputSurfaces = 1;
+
+						bitDepthMinus8 = OutPicture.DecodeCreateInfo.bitDepthMinus8 = SPS->bit_depth_luma_minus8;
+						ulWidth = OutPicture.DecodeCreateInfo.ulWidth = SPS->pic_width_in_luma_samples;
+						ulHeight = OutPicture.DecodeCreateInfo.ulHeight = SPS->pic_height_in_luma_samples;
+						ChromaFormat = OutPicture.DecodeCreateInfo.ChromaFormat = static_cast<cudaVideoChromaFormat>(SPS->chroma_format_idc.Value);
+						ulIntraDecodeOnly = OutPicture.DecodeCreateInfo.ulIntraDecodeOnly = static_cast<unsigned long>(SPS->profile_tier_level.general_constraint_flags & EH265ConstraintFlag::intra_constraint_flag);
+
+						ulTargetWidth = OutPicture.DecodeCreateInfo.ulTargetWidth = SPS->pic_width_in_luma_samples;
+						ulTargetHeight = OutPicture.DecodeCreateInfo.ulTargetHeight = SPS->pic_height_in_luma_samples;
+
+						// Need to scale the out frame depending on the chroma format
+						if (SPS->conformance_window_flag)
+						{
+							display_area = OutPicture.DecodeCreateInfo.display_area = {
+								static_cast<short>(SPS->conf_win.left_offset * (SPS->chroma_format_idc == 1 || SPS->chroma_format_idc == 2 ? 2 : 1)),
+								static_cast<short>(SPS->conf_win.top_offset * (SPS->chroma_format_idc == 1 ? 2 : 1)),
+								static_cast<short>(SPS->conf_win.right_offset * (SPS->chroma_format_idc == 1 || SPS->chroma_format_idc == 2 ? 2 : 1)),
+								static_cast<short>(SPS->conf_win.bottom_offset * (SPS->chroma_format_idc == 1 ? 2 : 1))
+							};
+
+							ulTargetWidth = OutPicture.DecodeCreateInfo.ulTargetWidth -= OutPicture.DecodeCreateInfo.display_area.left + OutPicture.DecodeCreateInfo.display_area.right;
+							ulTargetHeight = OutPicture.DecodeCreateInfo.ulTargetHeight -= OutPicture.DecodeCreateInfo.display_area.top + OutPicture.DecodeCreateInfo.display_area.bottom;
+						}
+
+						// NVDEC only supports 8 or 16 bit output surfaces in 420 or 444 format
+						if (ChromaFormat == cudaVideoChromaFormat_420)
+						{
+							OutputFormat = OutPicture.DecodeCreateInfo.OutputFormat = (bitDepthMinus8 == 0 ? cudaVideoSurfaceFormat_NV12 : cudaVideoSurfaceFormat_P016);
+						}
+						else
+						{
+							OutputFormat = OutPicture.DecodeCreateInfo.OutputFormat = (bitDepthMinus8 == 0 ? cudaVideoSurfaceFormat_YUV444 : cudaVideoSurfaceFormat_YUV444_16Bit);
+						}
+
+						OutPictureH265.pic_width_in_luma_samples = SPS->pic_width_in_luma_samples;
+						OutPictureH265.pic_height_in_luma_samples = SPS->pic_height_in_luma_samples;
+						OutPictureH265.log2_min_luma_coding_block_size_minus3 = SPS->log2_min_luma_coding_block_size_minus3;
+						OutPictureH265.log2_diff_max_min_luma_coding_block_size = SPS->log2_diff_max_min_luma_coding_block_size;
+						OutPictureH265.log2_min_transform_block_size_minus2 = SPS->log2_min_luma_transform_block_size_minus2;
+						OutPictureH265.log2_diff_max_min_transform_block_size = SPS->log2_diff_max_min_luma_transform_block_size;
+						OutPictureH265.pcm_enabled_flag = SPS->pcm_enabled_flag;
+						OutPictureH265.log2_min_pcm_luma_coding_block_size_minus3 = SPS->log2_min_pcm_luma_coding_block_size_minus3;
+						OutPictureH265.log2_diff_max_min_pcm_luma_coding_block_size = SPS->log2_diff_max_min_pcm_luma_coding_block_size;
+						OutPictureH265.pcm_sample_bit_depth_luma_minus1 = SPS->pcm_sample_bit_depth_luma_minus1;
+
+						OutPictureH265.pcm_sample_bit_depth_chroma_minus1 = SPS->pcm_sample_bit_depth_chroma_minus1;
+						OutPictureH265.pcm_loop_filter_disabled_flag = SPS->pcm_loop_filter_disabled_flag;
+						OutPictureH265.strong_intra_smoothing_enabled_flag = SPS->strong_intra_smoothing_enabled_flag;
+						OutPictureH265.max_transform_hierarchy_depth_intra = SPS->max_transform_hierarchy_depth_intra;
+						OutPictureH265.max_transform_hierarchy_depth_inter = SPS->max_transform_hierarchy_depth_inter;
+						OutPictureH265.amp_enabled_flag = SPS->amp_enabled_flag;
+						OutPictureH265.separate_colour_plane_flag = SPS->separate_colour_plane_flag;
+						OutPictureH265.log2_max_pic_order_cnt_lsb_minus4 = SPS->log2_max_pic_order_cnt_lsb_minus4;
+
+						OutPictureH265.num_short_term_ref_pic_sets = SPS->num_short_term_ref_pic_sets;
+						OutPictureH265.long_term_ref_pics_present_flag = SPS->long_term_ref_pics_present_flag;
+						OutPictureH265.num_long_term_ref_pics_sps = SPS->num_long_term_ref_pics_sps;
+						OutPictureH265.sps_temporal_mvp_enabled_flag = SPS->sps_temporal_mvp_enabled_flag;
+						OutPictureH265.sample_adaptive_offset_enabled_flag = SPS->sample_adaptive_offset_enabled_flag;
+						OutPictureH265.scaling_list_enable_flag = SPS->scaling_list_enabled_flag;
+						OutPictureH265.IrapPicFlag = (uint8)Slice->nal_unit_type.Value >= 16 && (uint8)Slice->nal_unit_type.Value <= 23;
+						OutPictureH265.IdrPicFlag = Slice->nal_unit_type == ENaluType::IDR_N_LP || Slice->nal_unit_type == ENaluType::IDR_W_RADL;
+
+						OutPictureH265.bit_depth_luma_minus8 = SPS->bit_depth_luma_minus8;
+						OutPictureH265.bit_depth_chroma_minus8 = SPS->bit_depth_chroma_minus8;
+
+						if (SPS->sps_range_extension_flag)
+						{
+							OutPictureH265.sps_range_extension_flag = SPS->sps_range_extension_flag;
+							OutPictureH265.high_precision_offsets_enabled_flag = SPS->sps_range_extension.high_precision_offsets_enabled_flag;
+							OutPictureH265.transform_skip_rotation_enabled_flag = SPS->sps_range_extension.transform_skip_rotation_enabled_flag;
+							OutPictureH265.transform_skip_context_enabled_flag = SPS->sps_range_extension.transform_skip_context_enabled_flag;
+							OutPictureH265.implicit_rdpcm_enabled_flag = SPS->sps_range_extension.implicit_rdpcm_enabled_flag;
+							OutPictureH265.explicit_rdpcm_enabled_flag = SPS->sps_range_extension.explicit_rdpcm_enabled_flag;
+							OutPictureH265.extended_precision_processing_flag = SPS->sps_range_extension.extended_precision_processing_flag;
+							OutPictureH265.intra_smoothing_disabled_flag = SPS->sps_range_extension.intra_smoothing_disabled_flag;
+							OutPictureH265.persistent_rice_adaptation_enabled_flag = SPS->sps_range_extension.persistent_rice_adaptation_enabled_flag;
+							OutPictureH265.cabac_bypass_alignment_enabled_flag = SPS->sps_range_extension.cabac_bypass_alignment_enabled_flag;
+						}
+
+						if (PPS->pps_range_extension_flag)
+						{
+							OutPictureH265.pps_range_extension_flag = PPS->pps_range_extension_flag;
+							OutPictureH265.log2_max_transform_skip_block_size_minus2 = PPS->pps_range_extension.log2_max_transform_skip_block_size_minus2;
+							OutPictureH265.log2_sao_offset_scale_luma = PPS->pps_range_extension.log2_sao_offset_scale_chroma;
+							OutPictureH265.log2_sao_offset_scale_chroma = PPS->pps_range_extension.log2_sao_offset_scale_chroma;
+							OutPictureH265.cross_component_prediction_enabled_flag = PPS->pps_range_extension.cross_component_prediction_enabled_flag;
+							OutPictureH265.chroma_qp_offset_list_enabled_flag = PPS->pps_range_extension.chroma_qp_offset_list_enabled_flag;
+							OutPictureH265.diff_cu_chroma_qp_offset_depth = PPS->pps_range_extension.diff_cu_chroma_qp_offset_depth;
+							OutPictureH265.chroma_qp_offset_list_len_minus1 = PPS->pps_range_extension.chroma_qp_offset_list_len_minus1;
+
+							FMemory::Memzero(OutPictureH265.cb_qp_offset_list, 21);
+							FMemory::Memcpy(OutPictureH265.cb_qp_offset_list, PPS->pps_range_extension.cb_qp_offset_list.GetData(), PPS->pps_range_extension.cb_qp_offset_list.Num());
+
+							FMemory::Memzero(OutPictureH265.cr_qp_offset_list, 21);
+							FMemory::Memcpy(OutPictureH265.cr_qp_offset_list, PPS->pps_range_extension.cr_qp_offset_list.GetData(), PPS->pps_range_extension.cr_qp_offset_list.Num());
+						}
+
+						OutPictureH265.dependent_slice_segments_enabled_flag = PPS->dependent_slice_segments_enabled_flag;
+						OutPictureH265.slice_segment_header_extension_present_flag = PPS->slice_segment_header_extension_present_flag;
+						OutPictureH265.sign_data_hiding_enabled_flag = PPS->sign_data_hiding_enabled_flag;
+						OutPictureH265.cu_qp_delta_enabled_flag = PPS->cu_qp_delta_enabled_flag;
+						OutPictureH265.diff_cu_qp_delta_depth = PPS->diff_cu_qp_delta_depth;
+						OutPictureH265.init_qp_minus26 = PPS->init_qp_minus26;
+						OutPictureH265.pps_cb_qp_offset = PPS->pps_cb_qp_offset;
+						OutPictureH265.pps_cr_qp_offset = PPS->pps_cr_qp_offset;
+
+						OutPictureH265.constrained_intra_pred_flag = PPS->constrained_intra_pred_flag;
+						OutPictureH265.weighted_pred_flag = PPS->weighted_pred_flag;
+						OutPictureH265.weighted_bipred_flag = PPS->weighted_bipred_flag;
+						OutPictureH265.transform_skip_enabled_flag = PPS->transform_skip_enabled_flag;
+						OutPictureH265.transquant_bypass_enabled_flag = PPS->transquant_bypass_enabled_flag;
+						OutPictureH265.entropy_coding_sync_enabled_flag = PPS->entropy_coding_sync_enabled_flag;
+						OutPictureH265.log2_parallel_merge_level_minus2 = PPS->log2_parallel_merge_level_minus2;
+						OutPictureH265.num_extra_slice_header_bits = PPS->num_extra_slice_header_bits;
+
+						OutPictureH265.loop_filter_across_tiles_enabled_flag = PPS->loop_filter_across_tiles_enabled_flag;
+						OutPictureH265.loop_filter_across_slices_enabled_flag = PPS->pps_loop_filter_across_slices_enabled_flag;
+						OutPictureH265.output_flag_present_flag = PPS->output_flag_present_flag;
+						OutPictureH265.num_ref_idx_l0_default_active_minus1 = PPS->num_ref_idx_l0_default_active_minus1;
+						OutPictureH265.num_ref_idx_l1_default_active_minus1 = PPS->num_ref_idx_l1_default_active_minus1;
+						OutPictureH265.lists_modification_present_flag = PPS->lists_modification_present_flag;
+						OutPictureH265.cabac_init_present_flag = PPS->cabac_init_present_flag;
+						OutPictureH265.pps_slice_chroma_qp_offsets_present_flag = PPS->pps_slice_chroma_qp_offsets_present_flag;
+
+						OutPictureH265.deblocking_filter_override_enabled_flag = PPS->deblocking_filter_override_enabled_flag;
+						OutPictureH265.pps_deblocking_filter_disabled_flag = PPS->pps_deblocking_filter_disabled_flag;
+						OutPictureH265.pps_beta_offset_div2 = PPS->pps_beta_offset_div2;
+						OutPictureH265.pps_tc_offset_div2 = PPS->pps_tc_offset_div2;
+						OutPictureH265.tiles_enabled_flag = PPS->tiles_enabled_flag;
+						OutPictureH265.uniform_spacing_flag = PPS->uniform_spacing_flag;
+						OutPictureH265.num_tile_columns_minus1 = PPS->num_tile_columns_minus1;
+						OutPictureH265.num_tile_rows_minus1 = PPS->num_tile_rows_minus1;
+
+						FMemory::Memzero(OutPictureH265.column_width_minus1, 21);
+						for (uint8 i = 0; i < PPS->column_width_minus1.Num(); i++)
+						{
+							OutPictureH265.column_width_minus1[i] = PPS->column_width_minus1[i];
+						}
+
+						FMemory::Memzero(OutPictureH265.row_height_minus1, 21);
+						for (uint8 i = 0; i < PPS->row_height_minus1.Num(); i++)
+						{
+							OutPictureH265.row_height_minus1[i] = PPS->row_height_minus1[i];
+						}
+
+						OutPictureH265.NumBitsForShortTermRPSInSlice = Slice->NumBitsForShortTermRPSInSlice;
+						OutPictureH265.NumDeltaPocsOfRefRpsIdx = SPS->short_term_ref_pic_sets[Slice->CurrRpsIdx].NumDeltaPocs;
+						OutPictureH265.NumPocTotalCurr = H265.ReferencePictureSet.NumPicTotalCurr;
+						OutPictureH265.NumPocStCurrBefore = H265.ReferencePictureSet.NumPocStCurrBefore;
+						OutPictureH265.NumPocStCurrAfter = H265.ReferencePictureSet.NumPocStCurrAfter;
+						OutPictureH265.NumPocLtCurr = H265.ReferencePictureSet.NumPocLtCurr;
+						OutPictureH265.CurrPicOrderCntVal = Slice->CurrPicOrderCntVal;
+
+						if (!OutPicture.intra_pic_flag)
+						{
+							for (uint8 i = 0; i < 16; i++)
+							{
+								if (i < SPS->MaxDpbSize)
+								{
+									OutPictureH265.RefPicIdx[i] = H265.ReferencePictureSet.PicUsage[RefPicIdx[i]] == EPictureUsage::UNUSED ? -1 : RefPicIdx[i];
+									OutPictureH265.IsLongTerm[i] = (H265.ReferencePictureSet.PicUsage[i] == EPictureUsage::LONG_TERM);
+								}
+								else
+								{
+									OutPictureH265.RefPicIdx[i] = -1;
+								}
+							}
+						}
+
+						
+						FMemory::Memcpy(OutPictureH265.PicOrderCntVal, H265.ReferencePictureSet.PicOrderCntVal, SPS->MaxDpbSize);
+						FMemory::Memcpy(OutPictureH265.RefPicSetStCurrBefore, H265.ReferencePictureSet.RefPicSetStCurrBefore, 8);
+						FMemory::Memcpy(OutPictureH265.RefPicSetStCurrAfter, H265.ReferencePictureSet.RefPicSetStCurrAfter, 8);
+						FMemory::Memcpy(OutPictureH265.RefPicSetLtCurr, H265.ReferencePictureSet.RefPicSetLtCurr, 8);
+
+
+
+						// HACK (aidan) we are not using layers yet so these should be zero until we do
+						FMemory::Memcpy(OutPictureH265.RefPicSetInterLayer0, H265.ReferencePictureSet.RefPicSetInterLayer0, 8);
+						FMemory::Memcpy(OutPictureH265.RefPicSetInterLayer1, H265.ReferencePictureSet.RefPicSetInterLayer1, 8);
+						
+						// TODO (aidan) check if you jump straight to default if the value is not present in the current active PPS or if you check SPS first
+						auto const& FillScalingLists = [&OutPictureH265](scaling_list_data_t const& ScalingListData) -> void {
+							for (uint8 i = 0; i < 6; i++)
+							{
+								if (ScalingListData.scaling_list_pred_mode_flag[0][i] > 0)
+								{
+									FMemory::Memcpy(OutPictureH265.ScalingList4x4[i], ScalingListData.ScalingList0[i], 16);
+								}
+								else
+								{
+									FMemory::Memcpy(OutPictureH265.ScalingList4x4[i], DefaultScalingList.ScalingList0[i], 16);
+								}
+							}
+
+							for (uint8 i = 0; i < 6; i++)
+							{
+								if (ScalingListData.scaling_list_pred_mode_flag[1][i] > 0)
+								{
+									FMemory::Memcpy(OutPictureH265.ScalingList8x8[i], ScalingListData.ScalingList1to3[0][i], 64);
+								}
+								else
+								{
+									FMemory::Memcpy(OutPictureH265.ScalingList8x8[i], DefaultScalingList.ScalingList1to3[0][i], 64);
+								}
+							}
+
+							for (uint8 i = 0; i < 6; i++)
+							{
+								if (ScalingListData.scaling_list_pred_mode_flag[2][i] > 0)
+								{
+									FMemory::Memcpy(OutPictureH265.ScalingList16x16[i], ScalingListData.ScalingList1to3[1][i], 64);
+									OutPictureH265.ScalingListDCCoeff16x16[i] = ScalingListData.scaling_list_dc_coef_minus8[2][i] + 8;
+								}
+								else
+								{
+									FMemory::Memcpy(OutPictureH265.ScalingList16x16[i], DefaultScalingList.ScalingList1to3[1][i], 64);
+									OutPictureH265.ScalingListDCCoeff16x16[i] = 16;
+								}
+							}
+
+							for (uint8 i = 0; i < 2; i++)
+							{
+								if (ScalingListData.scaling_list_pred_mode_flag[3][i] > 0)
+								{
+									FMemory::Memcpy(OutPictureH265.ScalingList32x32[i], ScalingListData.ScalingList1to3[2][i], 64);
+									OutPictureH265.ScalingListDCCoeff32x32[i] = ScalingListData.scaling_list_dc_coef_minus8[3][i] + 8;
+								}
+								else
+								{
+									FMemory::Memcpy(OutPictureH265.ScalingList32x32[i], DefaultScalingList.ScalingList1to3[2][i], 64);
+									OutPictureH265.ScalingListDCCoeff32x32[i] = 16;
+								}
+							}
+						};
+
+						if (PPS->pps_scaling_list_data_present_flag)
+						{
+							FillScalingLists(PPS->scaling_list_data);
+						}
+						else if (SPS->sps_scaling_list_data_present_flag)
+						{
+							FillScalingLists(SPS->scaling_list_data);
+						}
+						else
+						{
+							FillScalingLists(DefaultScalingList);
+						}
+						
+						OutPicture.SliceOffsets.Add(Slice->StartIdx);
+
+						// Could be a HACK (aidan) but reference NVDecoder parser seems to use this as a ring buffer
+						RefPicIdx[H265.CurrPicIdx % SPS->MaxDpbSize] = H265.CurrPicIdx % SPS->MaxDpbSize;
+						RefPicIdx[(H265.CurrPicIdx + 1) % SPS->MaxDpbSize] = -1;
+
+						bParsedFirstSlice = true;
+					}
+					else // Not first slice but member of current picture
+					{
+						OutPicture.SliceOffsets.Add(Slice->StartIdx); // NVDEC wants to see the 3 byte nal delimiter
+					}
+				}
+
+				if (OutPicture.SliceOffsets.Num() > 0)
+				{
+
+					// Make the offsets relative to the pointer to the first slice index in the bitstream
+					for (uint8 i = 1; i < OutPicture.SliceOffsets.Num(); i++)
+					{
+						OutPicture.SliceOffsets[i] -= OutPicture.SliceOffsets[0];
+					}
+
+					// Zero out the first slice index
+					OutPicture.SliceOffsets[0] = 0;
+
+					// Get size in bytes of the Bitstream with this PicParams slices					
+					OutPicture.nBitstreamDataLen = OutPicture.SliceOffsets.Last() - OutPicture.SliceOffsets[0] + Nalus[index - 1]->Size;
+
+					OutPicture.nNumSlices = OutPicture.SliceOffsets.Num();
+					OutPicture.pSliceDataOffsets = OutPicture.SliceOffsets.GetData();
+				}
+				else
+				{
+					OutPictures.Pop();
+				}
+				
+
+				// Loop through the remaining Nalus looking for the next slice
+				while (index < Nalus.Num() && !Nalus[index]->IsSlice())
+				{
+					index++;
+				}
+			} // while (index < Nalus.Num())
+
+			return EAVResult::Success;
+
+		} // case cudaVideoCodec_HEVC:
+		break;
+	default:
+		Result = FAVResult(EAVResult::ErrorMapping, FString::Printf(TEXT("Unsupported codec type %d"), CodecType), TEXT("NVDEC"));
+	}
+
+	return Result;
+}
+
+template <>
+DLLEXPORT FAVResult FAVExtension::TransformConfig(FVideoDecoderConfigNVDEC& OutConfig, FVideoDecoderConfig const& InConfig)
+{
+	OutConfig.Preset = InConfig.Preset;
+	OutConfig.LatencyMode = InConfig.LatencyMode;
+
+	return EAVResult::Success;
+}
+
+template <>
+DLLEXPORT FAVResult FAVExtension::TransformConfig(FVideoDecoderConfig& OutConfig, FVideoDecoderConfigNVDEC const& InConfig)
+{
+	OutConfig.Preset = InConfig.Preset;
+	OutConfig.LatencyMode = InConfig.LatencyMode;
+	
+	return EAVResult::Success;
+}
+
+template <>
+DLLEXPORT FAVResult FAVExtension::TransformConfig(FVideoDecoderConfigNVDEC& OutConfig, FVideoDecoderConfigH264 const& InConfig)
+{
+	OutConfig.CodecType = cudaVideoCodec_H264;
+	
+	return FAVExtension::TransformConfig<FVideoDecoderConfigNVDEC, FVideoDecoderConfig>(OutConfig, InConfig);
+}
+
+template <>
+DLLEXPORT FAVResult FAVExtension::TransformConfig(FVideoDecoderConfigNVDEC& OutConfig, FVideoDecoderConfigH265 const& InConfig)
+{
+	OutConfig.CodecType = cudaVideoCodec_HEVC;
+	
+	return FAVExtension::TransformConfig<FVideoDecoderConfigNVDEC, FVideoDecoderConfig>(OutConfig, InConfig);
+}
