Added new experimental plugin AVCodecs which handles software/hardware encoding and decoding of audio and video

diff --git a/Engine/Plugins/Experimental/AVCodecs/NVCodecs/Source/NVDEC/Private/Video/Decoders/VideoDecoderNVDEC.cpp b/Engine/Plugins/Experimental/AVCodecs/NVCodecs/Source/NVDEC/Private/Video/Decoders/VideoDecoderNVDEC.cpp
new file mode 100644
--- /dev/null
+++ b/Engine/Plugins/Experimental/AVCodecs/NVCodecs/Source/NVDEC/Private/Video/Decoders/VideoDecoderNVDEC.cpp
@@ -0,0 +1,333 @@
+// Copyright Epic Games, Inc. All Rights Reserved.
+
+#include "Video/Decoders/VideoDecoderNVDEC.h"
+
+#include "HAL/PlatformProcess.h"
+
+#include "NV12_to_BGRA8.cuh"
+#include "P010_to_ABGR10.cuh"
+
+FVideoDecoderNVDEC::~FVideoDecoderNVDEC()
+{
+	Close();
+}
+
+bool FVideoDecoderNVDEC::IsOpen() const
+{
+	return bIsOpen;
+}
+
+FAVResult FVideoDecoderNVDEC::Open(TSharedRef<FAVDevice> const& NewDevice, TSharedRef<FAVInstance> const& NewInstance)
+{
+	Close();
+
+	TVideoDecoder::Open(NewDevice, NewInstance);
+
+	bIsOpen = true;
+
+	return EAVResult::Success;
+}
+
+void FVideoDecoderNVDEC::Close()
+{
+	if (Decoder != nullptr)
+	{
+		FCUDAContextScope const ContextGuard(GetDevice()->GetContext<FVideoContextCUDA>()->Raw);
+
+		CUresult const Result = FAPI::Get<FNVDEC>().cuvidDestroyDecoder(Decoder);
+		if (Result != CUDA_SUCCESS)
+		{
+			FAVResult::Log(EAVResult::ErrorDestroying, TEXT("Failed to destroy NVENC encoder"), TEXT("NVDEC"), Result);
+		}
+
+		Decoder = nullptr;
+	}
+
+	bIsOpen = false;
+}
+
+bool FVideoDecoderNVDEC::IsInitialized() const
+{
+	return IsOpen() && Decoder != nullptr;
+}
+
+FAVResult FVideoDecoderNVDEC::ApplyConfig()
+{
+	if (IsOpen())
+	{
+		FVideoDecoderConfigNVDEC const& PendingConfig = GetPendingConfig();
+		if (AppliedConfig != PendingConfig)
+		{
+			if (IsInitialized())
+			{
+				/*
+				TODO(Andrew) Can be reconfigured? See https://docs.nvidia.com/video-technologies/video-codec-sdk/nvenc-video-encoder-api-prog-guide/#reconfigure-api
+				if (AppliedConfig.maxEncodeWidth == PendingConfig.maxEncodeWidth
+					&& AppliedConfig.maxEncodeHeight == PendingConfig.maxEncodeHeight
+					&& AppliedConfig.enablePTD == PendingConfig.enablePTD
+					&& AppliedConfig.enableEncodeAsync == PendingConfig.enableEncodeAsync
+					&& AppliedConfig.encodeConfig->gopLength == PendingConfig.encodeConfig->gopLength
+					&& AppliedConfig.encodeConfig->frameIntervalP == PendingConfig.encodeConfig->frameIntervalP
+					&& AppliedConfig.encodeConfig->encodeCodecConfig.h264Config.idrPeriod == PendingConfig.encodeConfig->encodeCodecConfig.h264Config.idrPeriod)
+				{
+					NV_ENC_STRUCT(NV_ENC_RECONFIGURE_PARAMS, ReconfigureParams);
+					FMemory::Memcpy(&ReconfigureParams.reInitEncodeParams, &PendingConfig, sizeof(NV_ENC_INITIALIZE_PARAMS));
+					ReconfigureParams.forceIDR = AppliedConfig.encodeWidth != PendingConfig.encodeWidth || AppliedConfig.encodeHeight != PendingConfig.encodeHeight;
+
+					NVENCSTATUS const Result = FAPI::Get<FNVENC_Context>().nvEncReconfigureEncoder(Encoder, &ReconfigureParams);
+					if (Result != NV_ENC_SUCCESS)
+					{
+						return FAVResult(EAVResult::Error, FString::Printf(TEXT("Failed to update NVENC encoder configuration: %s"), *FAPI::Get<FNVENC_Context>().GetErrorString(Encoder, Result)));
+					}
+				}
+				else
+				{
+					// TODO: Destroy and recreate with original session
+					unimplemented();
+				}
+				*/
+			}
+
+			if (!IsInitialized())
+			{
+				FCUDAContextScope const ContextGuard(GetDevice()->GetContext<FVideoContextCUDA>()->Raw);
+
+				CUresult const Result = FAPI::Get<FNVDEC>().cuvidCreateDecoder(&Decoder, &const_cast<FVideoDecoderConfigNVDEC&>(PendingConfig));
+				if (Result != CUDA_SUCCESS)
+				{
+					Close();
+
+					return FAVResult(EAVResult::ErrorCreating, TEXT("Failed to create decoder"), TEXT("NVDEC"), Result);
+				}
+			}
+		}
+
+		return TVideoDecoder::ApplyConfig();
+	}
+
+	return FAVResult(EAVResult::ErrorInvalidState, TEXT("Decoder not open"), TEXT("NVDEC"));
+}
+
+bool FVideoDecoderNVDEC::GetCapability(CUVIDDECODECAPS& CapsToQuery) const
+{
+	if (IsOpen())
+	{
+		FCUDAContextScope const ContextGuard(GetDevice()->GetContext<FVideoContextCUDA>()->Raw);
+
+		CUresult const Result = FAPI::Get<FNVDEC>().cuvidGetDecoderCaps(&CapsToQuery);
+		if (Result != CUDA_SUCCESS)
+		{
+			FAVResult::Log(EAVResult::Warning, TEXT("Failed to query for NVENC capability"), TEXT("NVDEC"), Result);
+
+			return false;
+		}
+
+		return CapsToQuery.bIsSupported > 0;
+	}
+
+	return FAVResult(EAVResult::ErrorInvalidState, TEXT("Decoder not open"), TEXT("NVDEC"));
+}
+
+FAVResult FVideoDecoderNVDEC::SendPacket(FVideoPacket const& Packet)
+{
+	if (!IsOpen())
+	{
+		return FAVResult(EAVResult::ErrorInvalidState, TEXT("Decoder not open"), TEXT("NVDEC"));
+	}
+
+	FVideoDecoderConfigNVDEC& PendingConfig = EditPendingConfig();
+
+	TArray<FVideoDecoderConfigNVDEC::FParsedPicture> Pictures;
+
+	FAVResult Result = PendingConfig.Parse(GetInstance().ToSharedRef(), Packet, Pictures);
+	if (Result.IsNotSuccess())
+	{
+		return Result;
+	}
+
+	if (Pictures.Num() == 0)
+	{
+		return EAVResult::PendingInput;
+	}
+
+	for (auto& Picture : Pictures)
+	{
+		if (FramesCount + 1 >= static_cast<int64>(Picture.DecodeCreateInfo.ulNumDecodeSurfaces))
+		{
+			return FAVResult(EAVResult::PendingOutput, TEXT("Decode buffer full"), TEXT("NVDEC"));
+		}
+
+		CUVIDPROCPARAMS MapParams = {};
+
+		FMemory::Memcpy(&static_cast<CUVIDDECODECREATEINFO&>(PendingConfig), &Picture.DecodeCreateInfo, sizeof(CUVIDDECODECREATEINFO));
+
+		Result = ApplyConfig();
+		if (Result.IsNotSuccess())
+		{
+			return Result;
+		}
+
+		CUresult const CUResult = FAPI::Get<FNVDEC>().cuvidDecodePicture(Decoder, &Picture);
+		if (CUResult != CUDA_SUCCESS)
+		{
+			return FAVResult(EAVResult::Error, TEXT("Failed to decode frame"), TEXT("NVDEC"), CUResult);
+		}
+
+		MapParams.progressive_frame = PendingConfig.DeinterlaceMode;
+		MapParams.second_field = Picture.second_field;
+		MapParams.top_field_first = Picture.bottom_field_flag != 0;
+
+		Frames.Enqueue({ Picture.CurrPicIdx, AppliedConfig.ulTargetWidth, AppliedConfig.ulTargetHeight, GetAppliedConfig().OutputFormat, MapParams });
+		++FramesCount;
+	}
+
+	return EAVResult::Success;
+}
+
+FAVResult FVideoDecoderNVDEC::ReceiveFrame(TResolvableVideoResource<FVideoResourceCUDA>& InOutResource)
+{
+	if (IsOpen())
+	{
+		FFrame Frame;
+		if (Frames.Peek(Frame))
+		{
+			FCUDAContextScope const ContextGuard(GetDevice()->GetContext<FVideoContextCUDA>()->Raw);
+
+			CUresult Result = CUDA_SUCCESS;
+			CUVIDGETDECODESTATUS DecodeStatus = {};
+
+			do
+			{
+				Result = FAPI::Get<FNVDEC>().cuvidGetDecodeStatus(Decoder, Frame.SurfaceIndex, &DecodeStatus);
+				if (Result != CUDA_SUCCESS)
+				{
+					return FAVResult(EAVResult::Error, TEXT("Could not retrieve decoder status"), TEXT("NVDEC"));
+				}
+				FPlatformProcess::Sleep(0.5f / 1000.0f); // HACK (aidan.possemiers) surely we can avoid a busy wait here
+			} while (DecodeStatus.decodeStatus == cuvidDecodeStatus_InProgress);
+
+			if (DecodeStatus.decodeStatus < cuvidDecodeStatus_Success)
+			{
+				return FAVResult(EAVResult::PendingInput, TEXT("Send more frames"), TEXT("NVDEC"), DecodeStatus.decodeStatus);
+			}
+
+			Frames.Pop();
+			--FramesCount;
+
+			if (DecodeStatus.decodeStatus > cuvidDecodeStatus_Success)
+			{
+				FAVResult::Log(EAVResult::Warning, TEXT("Error when decoding frame"), TEXT("NVDEC"), DecodeStatus.decodeStatus);
+			}
+			
+			// Map the decoder output buffer
+			CUdeviceptr MapSurface = 0;
+			uint32 MapPitch = 0;
+
+			Result = FAPI::Get<FNVDEC>().cuvidMapVideoFrame(Decoder, Frame.SurfaceIndex, &MapSurface, &MapPitch, &Frame.MapParams);
+			if (Result != CUDA_SUCCESS)
+			{
+				return FAVResult(EAVResult::ErrorMapping, TEXT("Failed to map decoded frame"), TEXT("NVDEC"), Result);
+			}
+
+			FVideoDescriptor ResourceDescriptor = {};
+
+#if 0 // Old path with automatic conversion
+			TFunction<CUresult(CUdeviceptr, CUarray, uint32, uint32, uint32)> ConversionFunction;
+
+			// TODO (aidan) handle conversion to non-RGB textures
+			switch (Frame.SurfaceFormat)
+			{
+			case cudaVideoSurfaceFormat_NV12:
+				ResourceDescriptor = FVideoDescriptor(EVideoFormat::BGRA, Frame.Width, Frame.Height);
+				ConversionFunction = NV12_to_BGRA8;
+				break;
+			case cudaVideoSurfaceFormat_P016:
+				ResourceDescriptor = FVideoDescriptor(EVideoFormat::ABGR10, Frame.Width, Frame.Height);
+				ConversionFunction = P010_to_ABGR10; // TODO
+				break;
+			/*case cudaVideoSurfaceFormat_YUV444:
+				ResourceDescriptor = FVideoDescriptor(EVideoFormat::YUV444, Frame.Width, Frame.Height);
+				break;
+			case cudaVideoSurfaceFormat_YUV444_16Bit:
+				ResourceDescriptor = FVideoDescriptor(EVideoFormat::YUV444_16, Frame.Width, Frame.Height);*/
+			default:
+				return FAVResult(EAVResult::ErrorUnsupported, FString::Printf(TEXT("Currently do not support the recieved OutputFormat (%d)"), Frame.SurfaceFormat));
+			}
+
+			if (!InOutResource.Resolve(GetDevice(), ResourceDescriptor))
+			{
+				return FAVResult(EAVResult::ErrorResolving, TEXT("Failed to resolve frame resource"), TEXT("NVDEC"));
+			}
+
+			InOutResource->Lock();
+
+			Result = ConversionFunction(MapSurface, InOutResource->GetRaw(), InOutResource->GetWidth(), InOutResource->GetHeight(), MapPitch);
+
+			InOutResource->Unlock();
+
+			if (Result != CUDA_SUCCESS)
+			{
+				return FAVResult(EAVResult::Error, TEXT("Failed to run CUDA kernal on decoded frame"), TEXT("NVDEC"), Result);
+			}
+
+			// START DEBUG
+
+#if DEBUG_DUMP_TO_DISK
+
+			TArray64<uint8> OutData;
+
+			InOutResource->ReadData(OutData);
+
+			FString SaveName = FString::Printf(TEXT("%s/DumpOutput/image%05d.%s"), *FPaths::ProjectSavedDir(), Frame.SurfaceIndex, Frame.SurfaceFormat == cudaVideoSurfaceFormat_NV12 ? TEXT("nv12") : TEXT("p016"));
+
+			FFileHelper::SaveArrayToFile(OutData, *SaveName);
+
+#endif
+
+			// END DEBUG
+#else // new code path just returns in the format that the decoder outputs and is up to the user to transform it
+			switch (Frame.SurfaceFormat)
+			{
+				case cudaVideoSurfaceFormat_NV12:
+					ResourceDescriptor = FVideoDescriptor(EVideoFormat::NV12, Frame.Width, Frame.Height);
+					break;
+				case cudaVideoSurfaceFormat_P016:
+					ResourceDescriptor = FVideoDescriptor(EVideoFormat::P010, Frame.Width, Frame.Height);
+					break;
+				case cudaVideoSurfaceFormat_YUV444:
+					ResourceDescriptor = FVideoDescriptor(EVideoFormat::YUV444, Frame.Width, Frame.Height);
+					break;
+				case cudaVideoSurfaceFormat_YUV444_16Bit:
+					ResourceDescriptor = FVideoDescriptor(EVideoFormat::YUV444_16, Frame.Width, Frame.Height);
+					break;
+				default:
+					return FAVResult(EAVResult::ErrorUnsupported, FString::Printf(TEXT("Currently do not support the recieved OutputFormat (%d)"), Frame.SurfaceFormat));
+			}
+
+			if (!InOutResource.Resolve(GetDevice(), ResourceDescriptor))
+			{
+				return FAVResult(EAVResult::ErrorResolving, TEXT("Failed to resolve frame resource"), TEXT("NVDEC"));
+			}
+
+			// Do copy into the VideoResource
+			InOutResource->CopyFrom(MapSurface, MapPitch);
+
+			// TODO (aidan.possemiers) this is messy as we may unmap the texture before the copy is done
+			// from what I can see there is no way to easily unmap this one the copy is done
+			Result = FAPI::Get<FNVDEC>().cuvidUnmapVideoFrame(Decoder, MapSurface);
+			if (Result != CUDA_SUCCESS)
+			{
+				return FAVResult(EAVResult::ErrorUnmapping, TEXT("Failed to unmap decoded frame"), TEXT("NVDEC"), Result);
+			}
+			
+
+#endif // end else !0		
+
+			return EAVResult::Success;
+		}
+
+		return EAVResult::PendingInput;
+	}
+
+	return FAVResult(EAVResult::ErrorInvalidState, TEXT("Decoder not open"), TEXT("NVDEC"));
+}
\ No newline at end of file
diff --git a/Engine/Plugins/Experimental/AVCodecs/NVCodecs/Source/NVDEC/Public/Video/Decoders/VideoDecoderNVDEC.h b/Engine/Plugins/Experimental/AVCodecs/NVCodecs/Source/NVDEC/Public/Video/Decoders/VideoDecoderNVDEC.h
new file mode 100644
--- /dev/null
+++ b/Engine/Plugins/Experimental/AVCodecs/NVCodecs/Source/NVDEC/Public/Video/Decoders/VideoDecoderNVDEC.h
@@ -0,0 +1,50 @@
+// Copyright Epic Games, Inc. All Rights Reserved.
+
+#pragma once
+
+#include "Video/VideoDecoder.h"
+
+#include "Containers/Queue.h"
+
+#include "Video/Resources/VideoResourceCUDA.h"
+
+#include "NVDEC.h"
+#include "Video/Decoders/Configs/VideoDecoderConfigNVDEC.h"
+
+class NVDEC_API FVideoDecoderNVDEC : public TVideoDecoder<FVideoResourceCUDA, FVideoDecoderConfigNVDEC>
+{
+private:
+	uint8 bIsOpen : 1;
+	CUvideodecoder Decoder = nullptr;
+
+	struct FFrame
+	{
+		int32 SurfaceIndex;
+		uint32 Width;
+		uint32 Height;
+		cudaVideoSurfaceFormat SurfaceFormat;
+		
+		CUVIDPROCPARAMS MapParams;
+	};
+
+	int32 FramesCount = 0;
+	TQueue<FFrame> Frames;
+	
+public:
+	FVideoDecoderNVDEC() = default;
+	virtual ~FVideoDecoderNVDEC() override;
+
+	virtual bool IsOpen() const override;
+	virtual FAVResult Open(TSharedRef<FAVDevice> const& NewDevice, TSharedRef<FAVInstance> const& NewInstance) override;
+	virtual void Close() override;
+
+	bool IsInitialized() const;
+
+	virtual FAVResult ApplyConfig() override;
+
+	bool GetCapability(CUVIDDECODECAPS& CapsToQuery) const;
+
+	virtual FAVResult SendPacket(FVideoPacket const& Packet) override;
+	
+	virtual FAVResult ReceiveFrame(TResolvableVideoResource<FVideoResourceCUDA>& InOutResource) override;
+};
diff --git a/Engine/Plugins/Experimental/AVCodecs/NVCodecs/Source/NVENC/Private/NVENCModule.cpp b/Engine/Plugins/Experimental/AVCodecs/NVCodecs/Source/NVENC/Private/NVENCModule.cpp
new file mode 100644
--- /dev/null
+++ b/Engine/Plugins/Experimental/AVCodecs/NVCodecs/Source/NVENC/Private/NVENCModule.cpp
@@ -0,0 +1,50 @@
+// Copyright Epic Games, Inc. All Rights Reserved.
+
+#include "Video/Encoders/VideoEncoderNVENC.h"
+#include "Video/Encoders/Configs/VideoEncoderConfigH264.h"
+#include "Video/Encoders/Configs/VideoEncoderConfigH265.h"
+#include "Video/Resources/VideoResourceVulkan.h"
+
+#if PLATFORM_WINDOWS
+#include "Video/Resources/Windows/VideoResourceD3D.h"
+#endif
+
+class FNVENCModule : public IModuleInterface
+{
+public:
+	virtual void StartupModule() override
+	{
+
+		/* Register NVENC+CUDA pathway for D3D12 and Vulkan. */
+		FVideoEncoder
+			::RegisterPermutationsOf<FVideoEncoderNVENCCUDA>
+			::With<
+				FVideoResourceCUDA,
+#if PLATFORM_WINDOWS
+				FVideoResourceD3D12,
+#endif
+				FVideoResourceVulkan>
+			::And<FVideoEncoderConfigNVENC, FVideoEncoderConfigH264, FVideoEncoderConfigH265>(
+				[](TSharedRef<FAVDevice> const& NewDevice, TSharedRef<FAVInstance> const& NewInstance)
+				{
+					return FAPI::Get<FNVENC>().IsValid() && NewDevice->HasContext<FVideoContextCUDA>();
+				});
+
+		/* Register NVENC + raw DX11 graphics device pathway for D3D11. 
+		* The reason we seperate these out in this way is because D3D11 + CUDA/NVENC does not encode UE textures properly due to how they are laid out.
+		*/
+#if PLATFORM_WINDOWS
+		FVideoEncoder
+			::RegisterPermutationsOf<FVideoEncoderNVENCD3D11>
+			::With<FVideoResourceD3D11>
+			::And<FVideoEncoderConfigNVENC, FVideoEncoderConfigH264, FVideoEncoderConfigH265>(
+				[](TSharedRef<FAVDevice> const& NewDevice, TSharedRef<FAVInstance> const& NewInstance)
+				{
+					return FAPI::Get<FNVENC>().IsValid() && NewDevice->HasContext<FVideoContextD3D11>();
+				});
+#endif // PLATFORM_WINDOWS
+	
+	}
+};
+
+IMPLEMENT_MODULE(FNVENCModule, NVENC);
diff --git a/Engine/Plugins/Experimental/AVCodecs/NVCodecs/Source/NVENC/Private/Video/Encoders/VideoEncoderNVENC.cpp b/Engine/Plugins/Experimental/AVCodecs/NVCodecs/Source/NVENC/Private/Video/Encoders/VideoEncoderNVENC.cpp
new file mode 100644
--- /dev/null
+++ b/Engine/Plugins/Experimental/AVCodecs/NVCodecs/Source/NVENC/Private/Video/Encoders/VideoEncoderNVENC.cpp
@@ -0,0 +1,446 @@
+// Copyright Epic Games, Inc. All Rights Reserved.
+
+#include "Video/Encoders/VideoEncoderNVENC.h"
+
+FEncoderNVENC::~FEncoderNVENC()
+{
+	Close();
+}
+
+bool FEncoderNVENC::IsOpen() const
+{
+	return Encoder != nullptr;
+}
+
+FAVResult FEncoderNVENC::Open(TSharedRef<FAVDevice> const& NewDevice, TSharedRef<FAVInstance> const& NewInstance, TFunction<void(NV_ENC_OPEN_ENCODE_SESSION_EX_PARAMS&)> SetupEncoderSessionFunc)
+{
+	Close();
+
+	NV_ENC_STRUCT(NV_ENC_OPEN_ENCODE_SESSION_EX_PARAMS, SessionParams);
+	SessionParams.apiVersion = NVENCAPI_VERSION;
+
+	SetupEncoderSessionFunc(SessionParams);
+
+	NVENCSTATUS const Result = FAPI::Get<FNVENC>().nvEncOpenEncodeSessionEx(&const_cast<NV_ENC_OPEN_ENCODE_SESSION_EX_PARAMS&>(SessionParams), &Encoder);
+	if (Result != NV_ENC_SUCCESS)
+	{
+		Close();
+
+		return FAVResult(EAVResult::ErrorCreating, TEXT("Failed to create encoder"), TEXT("NVENC"), Result);
+	}
+
+	return EAVResult::Success;
+}
+
+void FEncoderNVENC::Close()
+{
+	if (IsOpen())
+	{
+		if (Buffer != nullptr)
+		{
+			NVENCSTATUS const Result = FAPI::Get<FNVENC>().nvEncDestroyBitstreamBuffer(Encoder, Buffer);
+			if (Result != NV_ENC_SUCCESS)
+			{
+				FAVResult::Log(EAVResult::ErrorDestroying, TEXT("Failed to destroy output buffer"), TEXT("NVENC"), Result);
+			}
+
+			Buffer = nullptr;
+		}
+
+		if (Encoder != nullptr)
+		{
+			NVENCSTATUS const Result = FAPI::Get<FNVENC>().nvEncDestroyEncoder(Encoder);
+			if (Result != NV_ENC_SUCCESS)
+			{
+				FAVResult::Log(EAVResult::ErrorDestroying, TEXT("Failed to destroy encoder"), TEXT("NVENC"), Result);
+			}
+
+			Encoder = nullptr;
+		}
+	}
+}
+
+bool FEncoderNVENC::IsInitialized() const
+{
+	return IsOpen() && Buffer != nullptr;
+}
+
+FAVResult FEncoderNVENC::ApplyConfig(FVideoEncoderConfigNVENC const& AppliedConfig, FVideoEncoderConfigNVENC const& PendingConfig, TFunction<FAVResult()> ApplyConfigFunc)
+{
+	if (IsOpen())
+	{
+		// FVideoEncoderConfigNVENC const& PendingConfig = GetPendingConfig();
+		if (AppliedConfig != PendingConfig)
+		{
+			if (IsInitialized())
+			{
+				// Can be reconfigured? See https://docs.nvidia.com/video-technologies/video-codec-sdk/nvenc-video-encoder-api-prog-guide/#reconfigure-api
+				if (AppliedConfig.maxEncodeWidth == PendingConfig.maxEncodeWidth
+					&& AppliedConfig.maxEncodeHeight == PendingConfig.maxEncodeHeight
+					&& AppliedConfig.enablePTD == PendingConfig.enablePTD
+					&& AppliedConfig.enableEncodeAsync == PendingConfig.enableEncodeAsync
+					&& AppliedConfig.encodeConfig->gopLength == PendingConfig.encodeConfig->gopLength
+					&& AppliedConfig.encodeConfig->frameIntervalP == PendingConfig.encodeConfig->frameIntervalP
+					&& AppliedConfig.encodeConfig->encodeCodecConfig.h264Config.idrPeriod == PendingConfig.encodeConfig->encodeCodecConfig.h264Config.idrPeriod)
+				{
+					NV_ENC_STRUCT(NV_ENC_RECONFIGURE_PARAMS, ReconfigureParams);
+					FMemory::Memcpy(&ReconfigureParams.reInitEncodeParams, &static_cast<NV_ENC_INITIALIZE_PARAMS const&>(PendingConfig), sizeof(NV_ENC_INITIALIZE_PARAMS));
+					ReconfigureParams.forceIDR = AppliedConfig.encodeWidth != PendingConfig.encodeWidth || AppliedConfig.encodeHeight != PendingConfig.encodeHeight;
+
+					NVENCSTATUS const Result = FAPI::Get<FNVENC>().nvEncReconfigureEncoder(Encoder, &ReconfigureParams);
+					if (Result != NV_ENC_SUCCESS)
+					{
+						return FAVResult(EAVResult::Error, TEXT("Failed to update encoder configuration"), TEXT("NVENC"), Result);
+					}
+				}
+				else
+				{
+					// TODO: Destroy and recreate with original session
+					unimplemented();
+				}
+			}
+
+			if (!IsInitialized())
+			{
+				NVENCSTATUS Result = FAPI::Get<FNVENC>().nvEncInitializeEncoder(Encoder, &const_cast<FVideoEncoderConfigNVENC&>(PendingConfig));
+				if (Result != NV_ENC_SUCCESS)
+				{
+					return FAVResult(EAVResult::ErrorCreating, TEXT("Failed to initialize encoder"), TEXT("NVENC"), Result);
+				}
+
+				NV_ENC_STRUCT(NV_ENC_CREATE_BITSTREAM_BUFFER, CreateBuffer);
+
+				Result = FAPI::Get<FNVENC>().nvEncCreateBitstreamBuffer(Encoder, &CreateBuffer);
+				if (Result != NV_ENC_SUCCESS)
+				{
+					return FAVResult(EAVResult::ErrorCreating, TEXT("Failed to create output buffer"), TEXT("NVENC"), Result);
+				}
+
+				Buffer = CreateBuffer.bitstreamBuffer;
+			}
+		}
+
+		return ApplyConfigFunc();
+	}
+
+	return FAVResult(EAVResult::ErrorInvalidState, TEXT("Encoder not open"), TEXT("NVENC"));
+}
+
+int FEncoderNVENC::GetCapability(GUID EncodeGUID, NV_ENC_CAPS CapsToQuery) const
+{
+	if (IsOpen())
+	{
+		int CapsValue = 0;
+
+		NV_ENC_STRUCT(NV_ENC_CAPS_PARAM, CapsParam);
+		CapsParam.capsToQuery = CapsToQuery;
+
+		NVENCSTATUS const Result = FAPI::Get<FNVENC>().nvEncGetEncodeCaps(Encoder, EncodeGUID, &CapsParam, &CapsValue);
+		if (Result != NV_ENC_SUCCESS)
+		{
+			FAVResult::Log(EAVResult::Warning, FString::Printf(TEXT("Failed to query for capability %d"), CapsToQuery), TEXT("NVENC"), Result);
+
+			return 0;
+		}
+
+		return CapsValue;
+	}
+
+	return FAVResult(EAVResult::ErrorInvalidState, TEXT("Encoder not open"), TEXT("NVENC"));
+}
+
+FAVResult FEncoderNVENC::SendFrame(TSharedPtr<FVideoResource> const& Resource, uint32 Timestamp, bool bForceKeyframe, TFunction<FAVResult()> ApplyConfigFunc, TFunction<void(NV_ENC_REGISTER_RESOURCE&)> SetResourceToRegisterFunc)
+{
+	if (IsOpen())
+	{
+		FAVResult AVResult = ApplyConfigFunc();
+		if (AVResult.IsNotSuccess())
+		{
+			return AVResult;
+		}
+
+		NV_ENC_STRUCT(NV_ENC_MAP_INPUT_RESOURCE, MapResource);
+		NV_ENC_STRUCT(NV_ENC_REGISTER_RESOURCE, RegisterResource);
+		NV_ENC_STRUCT(NV_ENC_PIC_PARAMS, Picture);
+
+		NVENCSTATUS Result;
+		if (Resource.IsValid())
+		{
+			TAVResult<NV_ENC_BUFFER_FORMAT> const ConvertedPixelFormat = FVideoEncoderConfigNVENC::ConvertFormat(Resource->GetFormat());
+			if (ConvertedPixelFormat.IsNotSuccess())
+			{
+				return ConvertedPixelFormat;
+			}
+
+			// START DEBUG
+#if DEBUG_DUMP_TO_DISK
+			{
+				TArray64<uint8> OutData;
+
+				Resource->ReadData(OutData);
+
+				FString SaveName = FString::Printf(TEXT("%s/DumpInput/image%05d.%s"), *FPaths::ProjectSavedDir(), Timestamp, ConvertedPixelFormat == NV_ENC_BUFFER_FORMAT_NV12 ? TEXT("nv12") : TEXT("p016"));
+
+				FFileHelper::SaveArrayToFile(OutData, *SaveName);
+			}
+#endif
+			// END DEBUG
+
+			SetResourceToRegisterFunc(RegisterResource);
+			RegisterResource.width = Resource->GetWidth();
+			RegisterResource.height = Resource->GetHeight();
+			RegisterResource.pitch = Resource->GetStride();
+			RegisterResource.bufferFormat = ConvertedPixelFormat;
+			RegisterResource.bufferUsage = NV_ENC_INPUT_IMAGE;
+
+			Result = FAPI::Get<FNVENC>().nvEncRegisterResource(Encoder, &RegisterResource);
+			if (Result != NV_ENC_SUCCESS)
+			{
+				return FAVResult(EAVResult::ErrorMapping, TEXT("Failed to register frame resource"), TEXT("NVENC"), Result);
+			}
+
+			MapResource.registeredResource = RegisterResource.registeredResource;
+
+			Result = FAPI::Get<FNVENC>().nvEncMapInputResource(Encoder, &MapResource);
+			if (Result != NV_ENC_SUCCESS)
+			{
+				return FAVResult(EAVResult::ErrorMapping, TEXT("Failed to map frame resource"), TEXT("NVENC"), Result);
+			}
+
+			Picture.inputWidth = RegisterResource.width;
+			Picture.inputHeight = RegisterResource.height;
+			Picture.inputPitch = RegisterResource.pitch;
+			Picture.inputBuffer = MapResource.mappedResource;
+			Picture.bufferFmt = MapResource.mappedBufferFmt;
+			Picture.pictureStruct = NV_ENC_PIC_STRUCT_FRAME;
+			Picture.inputTimeStamp = Timestamp;
+
+			if (bForceKeyframe)
+			{
+				Picture.encodePicFlags |= NV_ENC_PIC_FLAG_FORCEIDR;
+			}
+			else // HACK (aidan.possemiers) AVCodecs h265 decoder only currently supports intra frames
+			{
+				Picture.encodePicFlags |= NV_ENC_PIC_FLAG_FORCEINTRA;
+			}
+		}
+		else
+		{
+			Picture.encodePicFlags = NV_ENC_PIC_FLAG_EOS;
+		}
+
+		if (Resource.IsValid())
+		{
+			Resource->Lock();
+		}
+
+		AVResult = SendFrame(Picture);
+
+		if (Resource.IsValid())
+		{
+			Resource->Unlock();
+		}
+
+		if (MapResource.mappedResource != nullptr)
+		{
+			Result = FAPI::Get<FNVENC>().nvEncUnmapInputResource(Encoder, MapResource.mappedResource);
+			if (Result != NV_ENC_SUCCESS)
+			{
+				FAVResult const UnmapResult = FAVResult(EAVResult::ErrorUnmapping, TEXT("Failed to unmap frame resource"), TEXT("NVENC"), Result);
+				if (UnmapResult < AVResult)
+				{
+					AVResult = UnmapResult;
+				}
+			}
+		}
+
+		if (RegisterResource.registeredResource != nullptr)
+		{
+			Result = FAPI::Get<FNVENC>().nvEncUnregisterResource(Encoder, RegisterResource.registeredResource);
+			if (Result != NV_ENC_SUCCESS)
+			{
+				FAVResult const UnmapResult = FAVResult(EAVResult::ErrorUnmapping, TEXT("Failed to unregister frame resource"), TEXT("NVENC"), Result);
+				if (UnmapResult < AVResult)
+				{
+					AVResult = UnmapResult;
+				}
+			}
+		}
+
+		return AVResult;
+	}
+
+	return FAVResult(EAVResult::ErrorInvalidState, TEXT("Encoder not open"), TEXT("NVENC"));
+}
+
+FAVResult FEncoderNVENC::SendFrame(NV_ENC_PIC_PARAMS Input)
+{
+	if (IsOpen())
+	{
+		if (Input.encodePicFlags & NV_ENC_PIC_FLAG_EOS)
+		{
+			NVENCSTATUS Result = FAPI::Get<FNVENC>().nvEncEncodePicture(Encoder, &Input);
+			if (Result != NV_ENC_SUCCESS)
+			{
+				return FAVResult(EAVResult::Error, TEXT("Error encoding end-of-stream picture"), TEXT("NVENC"), Result);
+			}
+
+			return EAVResult::Success;
+		}
+		else
+		{
+			if (Input.outputBitstream == nullptr)
+			{
+				Input.outputBitstream = Buffer;
+			}
+
+			NVENCSTATUS Result = FAPI::Get<FNVENC>().nvEncEncodePicture(Encoder, &Input);
+			if (Result != NV_ENC_SUCCESS)
+			{
+				return FAVResult(EAVResult::Error, TEXT("Error encoding picture"), TEXT("NVENC"), Result);
+			}
+
+			NV_ENC_STRUCT(NV_ENC_LOCK_BITSTREAM, BitstreamLock);
+			BitstreamLock.outputBitstream = Buffer;
+
+			Result = FAPI::Get<FNVENC>().nvEncLockBitstream(Encoder, &BitstreamLock);
+			if (Result != NV_ENC_SUCCESS)
+			{
+				return FAVResult(EAVResult::ErrorLocking, TEXT("Failed to lock output bitstream"), TEXT("NVENC"), Result);
+			}
+
+			TSharedPtr<uint8> const CopiedData = MakeShareable(new uint8[BitstreamLock.bitstreamSizeInBytes]);
+			FMemory::BigBlockMemcpy(CopiedData.Get(), BitstreamLock.bitstreamBufferPtr, BitstreamLock.bitstreamSizeInBytes);
+
+			Packets.Enqueue(
+				FVideoPacket(
+					CopiedData,
+					BitstreamLock.bitstreamSizeInBytes,
+					BitstreamLock.outputTimeStamp,
+					BitstreamLock.frameIdx,
+					BitstreamLock.frameAvgQP,
+					(BitstreamLock.pictureType & NV_ENC_PIC_TYPE_IDR) != 0));
+
+			Result = FAPI::Get<FNVENC>().nvEncUnlockBitstream(Encoder, BitstreamLock.outputBitstream);
+			if (Result != NV_ENC_SUCCESS)
+			{
+				return FAVResult(EAVResult::ErrorUnlocking, TEXT("Failed to unlock output buffer"), TEXT("NVENC"));
+			}
+
+			return EAVResult::Success;
+		}
+	}
+
+	return FAVResult(EAVResult::ErrorInvalidState, TEXT("Encoder not open"), TEXT("NVENC"));
+}
+
+#if PLATFORM_WINDOWS
+FAVResult FEncoderNVENC::SendFrameD3D11(TRefCountPtr<ID3D11Device> Device, TSharedPtr<FVideoResourceD3D11> const& Resource, uint32 Timestamp, bool bForceKeyframe, TFunction<FAVResult()> ApplyConfigFunc)
+{
+	TRefCountPtr<ID3D11Texture2D>& Tex = const_cast<TRefCountPtr<ID3D11Texture2D>&>(Resource->GetRaw());
+
+	// Very important: We create a new device for NVENC - we must share out D3D11 texture with the new device by "Opening" it to that device.
+	HRESULT Result = Device->OpenSharedResource(Resource->GetSharedHandle(), __uuidof(ID3D11Texture2D), (void**)(Tex.GetInitReference()));
+	if (Result != S_OK)
+	{
+		return FAVResult(EAVResult::Fatal, TEXT("Failed to open shared handle."), TEXT("NVENC"), Result);
+	}
+
+	return SendFrame(Resource, Timestamp, bForceKeyframe, ApplyConfigFunc, [Resource](NV_ENC_REGISTER_RESOURCE& RegisterResource) {
+		RegisterResource.resourceType = NV_ENC_INPUT_RESOURCE_TYPE_DIRECTX;
+		RegisterResource.resourceToRegister = Resource->GetRaw();
+	});
+}
+
+/**
+ * @returns The FString version of HRESULT error message so we can log it.
+ **/
+const FString AVGetComErrorDescription(HRESULT Res)
+{
+	const uint32 BufSize = 4096;
+	WIDECHAR Buffer[4096];
+	if (::FormatMessageW(FORMAT_MESSAGE_FROM_SYSTEM,
+			nullptr,
+			Res,
+			MAKELANGID(LANG_NEUTRAL, SUBLANG_NEUTRAL),
+			Buffer,
+			sizeof(Buffer) / sizeof(*Buffer),
+			nullptr))
+	{
+		return Buffer;
+	}
+	else
+	{
+		return TEXT("[cannot find error description]");
+	}
+}
+
+/**
+ * You may be asking yourself, why on Earth would we need to create a seperate D3D11 device when using NVENC?
+ * Well, it basically comes down to the D3D11 device created by Unreal Engine is not compatible with NVENC due to sdk version changes (we think?)
+ * and Unreal Engine is in no hurry to bump D3D11 version - so we create out own device. This works only because
+ * the D3D11 textures we are using are created as shared resources and accessed through shared handles.
+ * Without this chicainery the D3D11 device will be ejected and Unreal Engine and NVENC will crash.
+ */
+FAVResult FEncoderNVENC::CreateD3D11Device(TSharedRef<FAVDevice> const& InDevice, TRefCountPtr<ID3D11Device>& OutEncoderDevice, TRefCountPtr<ID3D11DeviceContext>& OutEncoderDeviceContext)
+{
+	TRefCountPtr<IDXGIDevice> DXGIDevice;
+	TRefCountPtr<IDXGIAdapter> Adapter;
+
+	HRESULT Result = InDevice->GetContext<FVideoContextD3D11>()->Device->QueryInterface(__uuidof(IDXGIDevice), (void**)DXGIDevice.GetInitReference());
+	if (Result != S_OK)
+	{
+		return FAVResult(EAVResult::Fatal, FString::Printf(TEXT("ID3D11Device::QueryInterface() failed 0x%X - %s."), Result, *AVGetComErrorDescription(Result)), TEXT("D3D11"));
+	}
+	else if ((Result = DXGIDevice->GetAdapter(Adapter.GetInitReference())) != S_OK)
+	{
+		return FAVResult(EAVResult::Fatal, FString::Printf(TEXT("DXGIDevice::GetAdapter() failed 0x%X - %s."), Result, *AVGetComErrorDescription(Result)), TEXT("D3D11"));
+	}
+
+	uint32 DeviceFlags = 0;
+	D3D_FEATURE_LEVEL FeatureLevel = D3D_FEATURE_LEVEL_11_0;
+	D3D_FEATURE_LEVEL ActualFeatureLevel;
+
+	if ((Result = D3D11CreateDevice(
+			 Adapter,
+			 D3D_DRIVER_TYPE_UNKNOWN,
+			 NULL,
+			 DeviceFlags,
+			 &FeatureLevel,
+			 1,
+			 D3D11_SDK_VERSION,
+			 OutEncoderDevice.GetInitReference(),
+			 &ActualFeatureLevel,
+			 OutEncoderDeviceContext.GetInitReference()))
+		!= S_OK)
+	{
+		return FAVResult(EAVResult::Fatal, FString::Printf(TEXT("D3D11CreateDevice() failed 0x%X - %s."), Result, *AVGetComErrorDescription(Result)), TEXT("D3D11"));
+	}
+
+	return FAVResult(EAVResult::Success, TEXT("Created D3D11 device for NVENC."), TEXT("D3D11"));
+}
+
+#endif // PLATFORM_WINDOWS
+
+FAVResult FEncoderNVENC::SendFrameCUDA(TSharedPtr<FVideoResourceCUDA> const& Resource, uint32 Timestamp, bool bForceKeyframe, TFunction<FAVResult()> ApplyConfigFunc)
+{
+	return SendFrame(Resource, Timestamp, bForceKeyframe, ApplyConfigFunc, [Resource](NV_ENC_REGISTER_RESOURCE& RegisterResource) {
+		RegisterResource.resourceType = NV_ENC_INPUT_RESOURCE_TYPE_CUDAARRAY;
+		RegisterResource.resourceToRegister = Resource->GetRaw();
+	});
+}
+
+FAVResult FEncoderNVENC::ReceivePacket(FVideoPacket& OutPacket)
+{
+	if (IsOpen())
+	{
+		if (Packets.Dequeue(OutPacket))
+		{
+			return EAVResult::Success;
+		}
+
+		return EAVResult::PendingInput;
+	}
+
+	return FAVResult(EAVResult::ErrorInvalidState, TEXT("Encoder not open"), TEXT("NVENC"));
+}
diff --git a/Engine/Plugins/Experimental/AVCodecs/NVCodecs/Source/NVENC/Public/NVENC.h b/Engine/Plugins/Experimental/AVCodecs/NVCodecs/Source/NVENC/Public/NVENC.h
new file mode 100644
--- /dev/null
+++ b/Engine/Plugins/Experimental/AVCodecs/NVCodecs/Source/NVENC/Public/NVENC.h
@@ -0,0 +1,43 @@
+// Copyright Epic Games, Inc. All Rights Reserved.
+
+#pragma once
+
+#include "AVUtility.h"
+
+THIRD_PARTY_INCLUDES_START
+
+#if PLATFORM_WINDOWS
+#include "Windows/AllowWindowsPlatformTypes.h"
+#include "Windows/PreWindowsApi.h"
+#endif
+
+#include <nvEncodeAPI.h>
+
+#if PLATFORM_WINDOWS
+#include "Windows/PostWindowsApi.h"
+#include "Windows/HideWindowsPlatformTypes.h"
+#endif
+
+THIRD_PARTY_INCLUDES_END
+
+#define NV_ENC_STRUCT(OfType, VarName)	\
+	OfType	VarName;					\
+	FMemory::Memzero(VarName);			\
+	VarName.version = OfType ## _VER
+
+class NVENC_API FNVENC : public FAPI, public NV_ENCODE_API_FUNCTION_LIST
+{
+public:
+	bool bHasCompatibleGPU = true;
+
+	FNVENC();
+
+	virtual bool IsValid() const override;
+
+	FString GetErrorString(void* Encoder, NVENCSTATUS ForStatus) const;
+
+private:
+	void* DllHandle = nullptr;
+};
+
+DECLARE_TYPEID(FNVENC, NVENC_API);
diff --git a/Engine/Plugins/Experimental/AVCodecs/NVCodecs/Source/NVENC/Public/Video/Encoders/VideoEncoderNVENC.h b/Engine/Plugins/Experimental/AVCodecs/NVCodecs/Source/NVENC/Public/Video/Encoders/VideoEncoderNVENC.h
new file mode 100644
--- /dev/null
+++ b/Engine/Plugins/Experimental/AVCodecs/NVCodecs/Source/NVENC/Public/Video/Encoders/VideoEncoderNVENC.h
@@ -0,0 +1,151 @@
+// Copyright Epic Games, Inc. All Rights Reserved.
+
+#pragma once
+
+#include "Video/VideoEncoder.h"
+#include "Video/Resources/VideoResourceCUDA.h"
+#include "Video/VideoResource.h"
+#include "NVENC.h"
+#include "Video/Encoders/Configs/VideoEncoderConfigNVENC.h"
+#include "Templates/RefCounting.h"
+#include "Containers/Queue.h"
+
+#if PLATFORM_WINDOWS
+	#include "Video/Resources/Windows/VideoResourceD3D.h"
+#endif // PLATFORM_WINDOWS
+
+/**
+ * This class is the core of the NVENC encoder logic.
+ * It is embedded inside NVCodecEncoders using composition instead of inheritance because templates + inheritance was bit messy here where we need concrete types.
+ **/
+class NVENC_API FEncoderNVENC
+{
+private:
+	void* Encoder = nullptr;
+	NV_ENC_OUTPUT_PTR Buffer = nullptr;
+	TQueue<FVideoPacket> Packets;
+
+public:
+	// Begin matching the TVideoEncoder interface
+	virtual ~FEncoderNVENC();
+	virtual bool IsOpen() const;
+	virtual FAVResult Open(TSharedRef<FAVDevice> const& NewDevice, TSharedRef<FAVInstance> const& NewInstance, TFunction<void(NV_ENC_OPEN_ENCODE_SESSION_EX_PARAMS&)> SetupEncoderSessionFunc);
+	virtual void Close();
+	virtual FAVResult ApplyConfig(FVideoEncoderConfigNVENC const& AppliedConfig, FVideoEncoderConfigNVENC const& PendingConfig, TFunction<FAVResult()> ApplyConfigFunc);
+#if PLATFORM_WINDOWS
+	virtual FAVResult CreateD3D11Device(TSharedRef<FAVDevice> const& InDevice, TRefCountPtr<ID3D11Device>& OutEncoderDevice, TRefCountPtr<ID3D11DeviceContext>& OutEncoderDeviceContext);
+	virtual FAVResult SendFrameD3D11(TRefCountPtr<ID3D11Device> Device, TSharedPtr<FVideoResourceD3D11> const& Resource, uint32 Timestamp, bool bForceKeyframe, TFunction<FAVResult()> ApplyConfigFunc);
+#endif // PLATFORM_WINDOWS
+	virtual FAVResult SendFrameCUDA(TSharedPtr<FVideoResourceCUDA> const& Resource, uint32 Timestamp, bool bForceKeyframe, TFunction<FAVResult()> ApplyConfigFunc);
+	virtual FAVResult SendFrame(TSharedPtr<FVideoResource> const& Resource, uint32 Timestamp, bool bForceKeyframe, TFunction<FAVResult()> ApplyConfigFunc, TFunction<void(NV_ENC_REGISTER_RESOURCE&)> SetResourceToRegisterFunc);
+	virtual FAVResult ReceivePacket(FVideoPacket& OutPacket);
+	// End matching the TVideoEncoder interface
+
+	bool IsInitialized() const;
+	FAVResult SendFrame(NV_ENC_PIC_PARAMS Input);
+	int GetCapability(GUID EncodeGUID, NV_ENC_CAPS CapsToQuery) const;
+};
+
+/**
+ *  An AVCodecs VideoEncoder that takes a FVideoResourceCUDA.
+ **/
+class NVENC_API FVideoEncoderNVENCCUDA : public TVideoEncoder<FVideoResourceCUDA, FVideoEncoderConfigNVENC>
+{
+private:
+	TUniquePtr<FEncoderNVENC> Base;
+
+public:
+	FVideoEncoderNVENCCUDA()
+		: Base(MakeUnique<FEncoderNVENC>()) {}
+
+	// Begin TVideoEncoder interface
+	virtual ~FVideoEncoderNVENCCUDA() = default;
+	virtual bool IsOpen() const override { return Base->IsOpen(); }
+	virtual void Close() override { Base->Close(); };
+	virtual FAVResult ReceivePacket(FVideoPacket& OutPacket) override { return Base->ReceivePacket(OutPacket); }
+
+	virtual FAVResult ApplyConfig() override
+	{
+		return Base->ApplyConfig(AppliedConfig, GetPendingConfig(), [this]() {
+			return TVideoEncoder::ApplyConfig();
+		});
+	}
+
+	virtual FAVResult Open(TSharedRef<FAVDevice> const& NewDevice, TSharedRef<FAVInstance> const& NewInstance) override
+	{
+		return Base->Open(NewDevice, NewInstance, [this, NewDevice, NewInstance](NV_ENC_OPEN_ENCODE_SESSION_EX_PARAMS& SessionParams) {
+			TVideoEncoder::Open(NewDevice, NewInstance);
+			SessionParams.deviceType = NV_ENC_DEVICE_TYPE_CUDA;
+			SessionParams.device = GetDevice()->GetContext<FVideoContextCUDA>()->Raw;
+		});
+	}
+
+	virtual FAVResult SendFrame(TSharedPtr<FVideoResourceCUDA> const& Resource, uint32 Timestamp, bool bForceKeyframe = false) override
+	{
+		return Base->SendFrameCUDA(Resource, Timestamp, bForceKeyframe, [this]() {
+			return ApplyConfig();
+		});
+	}
+	// End TVideoEncoder interface
+};
+
+#if PLATFORM_WINDOWS
+
+/**
+ *  An AVCodecs VideoEncoder that takes a FVideoResourceD3D11.
+ **/
+class NVENC_API FVideoEncoderNVENCD3D11 : public TVideoEncoder<FVideoResourceD3D11, FVideoEncoderConfigNVENC>
+{
+private:
+	TUniquePtr<FEncoderNVENC> Base;
+	TRefCountPtr<ID3D11Device> EncoderDevice;
+	TRefCountPtr<ID3D11DeviceContext> EncoderDeviceContext;
+
+public:
+	FVideoEncoderNVENCD3D11()
+		: Base(MakeUnique<FEncoderNVENC>()) {}
+
+	// Begin TVideoEncoder interface
+	virtual ~FVideoEncoderNVENCD3D11() = default;
+	virtual bool IsOpen() const override { return Base->IsOpen(); }
+	virtual void Close() override { Base->Close(); };
+	virtual FAVResult ReceivePacket(FVideoPacket& OutPacket) override { return Base->ReceivePacket(OutPacket); }
+
+	virtual FAVResult ApplyConfig() override
+	{
+		return Base->ApplyConfig(AppliedConfig, GetPendingConfig(), [this]() {
+			return TVideoEncoder::ApplyConfig();
+		});
+	};
+
+	virtual FAVResult Open(TSharedRef<FAVDevice> const& NewDevice, TSharedRef<FAVInstance> const& NewInstance) override
+	{
+
+		// then pass device in SendFrameD3D11 and do OpenHandle
+
+		return Base->Open(NewDevice, NewInstance, [this, NewDevice, NewInstance](NV_ENC_OPEN_ENCODE_SESSION_EX_PARAMS& SessionParams) {
+			FAVResult Result = Base->CreateD3D11Device(NewDevice, EncoderDevice, EncoderDeviceContext);
+
+			// If we failed to create device early out.
+			if (Result.IsNotSuccess())
+			{
+				return Result;
+			}
+
+			TVideoEncoder::Open(NewDevice, NewInstance);
+			SessionParams.deviceType = NV_ENC_DEVICE_TYPE_DIRECTX;
+			SessionParams.device = EncoderDevice;
+			// SessionParams.device = GetDevice()->GetContext<FVideoContextD3D11>()->Device;
+			return Result;
+		});
+	}
+
+	virtual FAVResult SendFrame(TSharedPtr<FVideoResourceD3D11> const& Resource, uint32 Timestamp, bool bForceKeyframe = false)
+	{
+		return Base->SendFrameD3D11(EncoderDevice, Resource, Timestamp, bForceKeyframe, [this]() {
+			return ApplyConfig();
+		});
+	}
+	// End TVideoEncoder interface
+};
+#endif // PLATFORM_WINDOWS
\ No newline at end of file
diff --git a/Engine/Plugins/Experimental/AVCodecs/QuickSyncCodecs/QuickSyncCodecs.uplugin b/Engine/Plugins/Experimental/AVCodecs/QuickSyncCodecs/QuickSyncCodecs.uplugin
new file mode 100644
--- /dev/null
+++ b/Engine/Plugins/Experimental/AVCodecs/QuickSyncCodecs/QuickSyncCodecs.uplugin
@@ -0,0 +1,28 @@
+{
+	"FileVersion": 3,
+	"Version": 1,
+	"VersionName": "0.1",
+	"FriendlyName": "QuickSyncCodecs",
+	"Description": "Adds codecs from the Intel QuickSync SDK to AVCodecs",
+	"Category": "Codecs",
+	"CreatedBy": "Epic Games, Inc.",
+	"CreatedByURL": "http://epicgames.com",
+	"MarketplaceURL": "",
+	"SupportURL": "",
+	"EnabledByDefault": false,
+	"CanContainContent": false,
+	"IsExperimentalVersion": true,
+	"Installed": true,
+	"Hidden": true,
+	"Modules": [ ],
+	"Plugins": [
+		{
+			"Name": "AVCodecsCore",
+			"Enabled": true
+		}
+	],
+	"SupportedTargetPlatforms": [
+		"Win64",
+		"Linux"
+	]
+}
\ No newline at end of file
diff --git a/Engine/Plugins/Experimental/AVCodecs/VPXCodecs/VPXCodecs.uplugin b/Engine/Plugins/Experimental/AVCodecs/VPXCodecs/VPXCodecs.uplugin
new file mode 100644
--- /dev/null
+++ b/Engine/Plugins/Experimental/AVCodecs/VPXCodecs/VPXCodecs.uplugin
@@ -0,0 +1,28 @@
+{
+	"FileVersion": 3,
+	"Version": 1,
+	"VersionName": "0.1",
+	"FriendlyName": "VPXCodecs",
+	"Description": "Adds codecs from the libvpx Codec Library to AVCodecs",
+	"Category": "Codecs",
+	"CreatedBy": "Epic Games, Inc.",
+	"CreatedByURL": "http://epicgames.com",
+	"MarketplaceURL": "",
+	"SupportURL": "",
+	"EnabledByDefault": false,
+	"CanContainContent": false,
+	"IsExperimentalVersion": true,
+	"Installed": true,
+	"Hidden": true,
+	"Modules": [ ],
+	"Plugins": [
+		{
+			"Name": "AVCodecsCore",
+			"Enabled": true
+		}
+	],
+	"SupportedTargetPlatforms": [
+		"Win64",
+		"Linux"
+	]
+}
\ No newline at end of file
diff --git a/Engine/Plugins/Experimental/AVCodecs/VideoToolboxCodecs/VideoToolboxCodecs.uplugin b/Engine/Plugins/Experimental/AVCodecs/VideoToolboxCodecs/VideoToolboxCodecs.uplugin
new file mode 100644
--- /dev/null
+++ b/Engine/Plugins/Experimental/AVCodecs/VideoToolboxCodecs/VideoToolboxCodecs.uplugin
@@ -0,0 +1,28 @@
+{
+	"FileVersion": 3,
+	"Version": 1,
+	"VersionName": "0.1",
+	"FriendlyName": "VideoToolboxCodecs",
+	"Description": "Adds codecs from the Apple Video Toolbox to AVCodecs",
+	"Category": "Codecs",
+	"CreatedBy": "Epic Games, Inc.",
+	"CreatedByURL": "http://epicgames.com",
+	"MarketplaceURL": "",
+	"SupportURL": "",
+	"EnabledByDefault": false,
+	"CanContainContent": false,
+	"IsExperimentalVersion": true,
+	"Installed": true,
+	"Hidden": true,
+	"Modules": [ ],
+	"Plugins": [
+		{
+			"Name": "AVCodecsCore",
+			"Enabled": true
+		}
+	],
+	"SupportedTargetPlatforms": [
+		"Mac",
+		"IOS"
+	]
+}
\ No newline at end of file
