Disable AVCodecsCore for Apple platforms

diff --git a/Engine/Plugins/Experimental/AVCodecs/AVCodecsCore/AVCodecsCore.uplugin b/Engine/Plugins/Experimental/AVCodecs/AVCodecsCore/AVCodecsCore.uplugin
--- a/Engine/Plugins/Experimental/AVCodecs/AVCodecsCore/AVCodecsCore.uplugin
+++ b/Engine/Plugins/Experimental/AVCodecs/AVCodecsCore/AVCodecsCore.uplugin
@@ -17,7 +17,8 @@
 		{
 			"Name": "AVCodecsCore",
 			"Type": "Runtime",
-			"LoadingPhase": "PostEngineInit"
+			"LoadingPhase": "PostEngineInit",
+			"PlatformAllowList": [ "Win64",	"Linux" ]
 		},
 		{
 			"Name": "AVCodecsCoreRHI",
@@ -25,7 +26,8 @@
 			"LoadingPhase": "PostConfigInit",
 			"TargetDenyList": [
 				"Program"
-			]
+			],
+			"PlatformAllowList": [ "Win64",	"Linux"	]
 		}
 	],
 	"SupportedTargetPlatforms": [
diff --git a/Engine/Plugins/Experimental/AVCodecs/NVCodecs/NVCodecs.uplugin b/Engine/Plugins/Experimental/AVCodecs/NVCodecs/NVCodecs.uplugin
--- a/Engine/Plugins/Experimental/AVCodecs/NVCodecs/NVCodecs.uplugin
+++ b/Engine/Plugins/Experimental/AVCodecs/NVCodecs/NVCodecs.uplugin
@@ -20,7 +20,8 @@
 			"LoadingPhase": "PostConfigInit",
 			"TargetDenyList": [
 				"Server"
-			]
+			],
+			"PlatformAllowList": [ "Win64",	"Linux" ]
 		},
 		{
 			"Name": "NVCodecsRHI",
@@ -28,7 +29,8 @@
 			"LoadingPhase": "PostConfigInit",
 			"TargetDenyList": [
 				"Server"
-			]
+			],
+			"PlatformAllowList": [ "Win64",	"Linux" ]
 		},
 		{
 			"Name": "NVDEC",
@@ -36,7 +38,8 @@
 			"LoadingPhase": "PostConfigInit",
 			"TargetDenyList": [
 				"Server"
-			]
+			],
+			"PlatformAllowList": [ "Win64",	"Linux" ]
 		},
 		{
 			"Name": "NVENC",
@@ -44,7 +47,8 @@
 			"LoadingPhase": "PostConfigInit",
 			"TargetDenyList": [
 				"Server"
-			]
+			],
+			"PlatformAllowList": [ "Win64",	"Linux" ]
 		}
 	],
 	"Plugins": [
diff --git a/Engine/Plugins/Experimental/AVCodecs/QuickSyncCodecs/QuickSyncCodecs.uplugin b/Engine/Plugins/Experimental/AVCodecs/QuickSyncCodecs/QuickSyncCodecs.uplugin
deleted file mode 100644
diff --git a/Engine/Plugins/Experimental/AVCodecs/VPXCodecs/VPXCodecs.uplugin b/Engine/Plugins/Experimental/AVCodecs/VPXCodecs/VPXCodecs.uplugin
deleted file mode 100644
diff --git a/Engine/Plugins/Experimental/AVCodecs/WMFCodecs/WMFCodecs.uplugin b/Engine/Plugins/Experimental/AVCodecs/WMFCodecs/WMFCodecs.uplugin
--- a/Engine/Plugins/Experimental/AVCodecs/WMFCodecs/WMFCodecs.uplugin
+++ b/Engine/Plugins/Experimental/AVCodecs/WMFCodecs/WMFCodecs.uplugin
@@ -33,7 +33,6 @@
 		}
 	],
 	"SupportedTargetPlatforms": [
-		"Win64",
-		"Hololens"
+		"Win64"
 	]
 }
\ No newline at end of file
