Prevent crashes caused by `PixelStreamingInput using `SlateApplication` before it is ready

NVCodecs packages binaries it needs as a part of project packaging, but in a disallowed way, disable this behaviour for now

diff --git a/Engine/Plugins/Experimental/AVCodecs/NVCodecs/Source/NVCodecs/NVCodecs.Build.cs b/Engine/Plugins/Experimental/AVCodecs/NVCodecs/Source/NVCodecs/NVCodecs.Build.cs
--- a/Engine/Plugins/Experimental/AVCodecs/NVCodecs/Source/NVCodecs/NVCodecs.Build.cs
+++ b/Engine/Plugins/Experimental/AVCodecs/NVCodecs/Source/NVCodecs/NVCodecs.Build.cs
@@ -29,11 +29,11 @@ public class NVCodecs : ModuleRules
 		// TODO make some way to compile the kernels for easy testing
 		/*nvcc -arch=sm_50 -gencode=arch=compute_50,code=sm_50 -gencode=arch=compute_52,code=sm_52 -gencode=arch=compute_60,code=sm_60 -gencode=arch=compute_61,code=sm_61 -gencode=arch=compute_70,code=sm_70 -gencode=arch=compute_75,code=sm_75 -gencode=arch=compute_75,code=compute_75 -o NV12_to_BGRA8.fatbin -fatbin --device-debug ../src/NV12_to_BGRA8.cu*/
 
-		RuntimeDependencies.Add(Path.Join("$(ProjectDir)", "Binaries", Target.Platform.ToString(), "Kernels", "NV12_to_BGRA8.fatbin"),
-								Path.Join(ModuleDirectory, "Kernels", "lib", "NV12_to_BGRA8.fatbin"));
+		// RuntimeDependencies.Add(Path.Join("$(ProjectDir)", "Binaries", Target.Platform.ToString(), "Kernels", "NV12_to_BGRA8.fatbin"),
+		// 						Path.Join(ModuleDirectory, "Kernels", "lib", "NV12_to_BGRA8.fatbin"));
 
-		RuntimeDependencies.Add(Path.Join("$(ProjectDir)", "Binaries", Target.Platform.ToString(), "Kernels", "P010_to_ABGR10.fatbin"),
-								Path.Join(ModuleDirectory, "Kernels", "lib", "P010_to_ABGR10.fatbin"));
+		// RuntimeDependencies.Add(Path.Join("$(ProjectDir)", "Binaries", Target.Platform.ToString(), "Kernels", "P010_to_ABGR10.fatbin"),
+		// 						Path.Join(ModuleDirectory, "Kernels", "lib", "P010_to_ABGR10.fatbin"));
 		
 		PrivateDependencyModuleNames.Add("Vulkan");
 		AddEngineThirdPartyPrivateStaticDependencies(Target, "Vulkan");
diff --git a/Engine/Plugins/Media/PixelStreaming/Source/PixelStreamingInput/Private/PixelStreamingInputModule.cpp b/Engine/Plugins/Media/PixelStreaming/Source/PixelStreamingInput/Private/PixelStreamingInputModule.cpp
--- a/Engine/Plugins/Media/PixelStreaming/Source/PixelStreamingInput/Private/PixelStreamingInputModule.cpp
+++ b/Engine/Plugins/Media/PixelStreaming/Source/PixelStreamingInput/Private/PixelStreamingInputModule.cpp
@@ -11,6 +11,11 @@ namespace UE::PixelStreamingInput
 {
 	void FPixelStreamingInputModule::StartupModule()
 	{
+		if (!FSlateApplication::IsInitialized())
+		{
+			return;
+		}
+
 		TSharedPtr<FPixelStreamingApplicationWrapper> PixelStreamerApplicationWrapper = MakeShareable(new FPixelStreamingApplicationWrapper(FSlateApplication::Get().GetPlatformApplication()));
 		TSharedPtr<FGenericApplicationMessageHandler> BaseHandler = FSlateApplication::Get().GetPlatformApplication()->GetMessageHandler();
 		InputHandler = MakeShared<FPixelStreamingInputHandler>(PixelStreamerApplicationWrapper, BaseHandler);
