Fix streaming failure at runtime: Fix logic error in NVENC config comparison in AVCodecs

diff --git a/Engine/Plugins/Experimental/AVCodecs/NVCodecs/Source/NVENC/Public/Video/Encoders/Configs/VideoEncoderConfigNVENC.h b/Engine/Plugins/Experimental/AVCodecs/NVCodecs/Source/NVENC/Public/Video/Encoders/Configs/VideoEncoderConfigNVENC.h
--- a/Engine/Plugins/Experimental/AVCodecs/NVCodecs/Source/NVENC/Public/Video/Encoders/Configs/VideoEncoderConfigNVENC.h
+++ b/Engine/Plugins/Experimental/AVCodecs/NVCodecs/Source/NVENC/Public/Video/Encoders/Configs/VideoEncoderConfigNVENC.h
@@ -64,12 +64,20 @@ public:
 
 	bool operator==(FVideoEncoderConfigNVENC const& Other) const
 	{
+		// We want to compare the encodeConfig member's contents not its memory address
+		// so we temporarily store "this" encodeConfig and set "this" encodeConfig to
+		// the same object as the one in Other, so it is essentially not considered in
+		// a memcmp of "this" and "Other"
 		NV_ENC_CONFIG* const TempEncodeConfig = encodeConfig;
 		const_cast<FVideoEncoderConfigNVENC*>(this)->encodeConfig = Other.encodeConfig;
 
-		bool const IsEqual = !FMemory::Memcmp(this, &Other, sizeof(FVideoEncoderConfigNVENC))
-			|| !FMemory::Memcmp(encodeConfig, Other.encodeConfig, sizeof(NV_ENC_CONFIG));
+		// Compare the data of "this" and Other, but also the data of 
+		// "this"->encodeConfig (currently stored in TempEncodeConfig) and
+		// Other.encodeConfig, true if both comparison find the data to be identical
+		bool const IsEqual = FMemory::Memcmp(this, &Other, sizeof(FVideoEncoderConfigNVENC)) == 0
+			&& FMemory::Memcmp(TempEncodeConfig, Other.encodeConfig, sizeof(NV_ENC_CONFIG)) == 0;
 
+		// Set "this" encodeConfig back to its original value
 		const_cast<FVideoEncoderConfigNVENC*>(this)->encodeConfig = TempEncodeConfig;
 
 		return IsEqual;
