Fix: Make H26X not send only intra frames

diff --git a/Engine/Plugins/Experimental/AVCodecs/NVCodecs/Source/NVENC/Private/Video/Encoders/VideoEncoderNVENC.cpp b/Engine/Plugins/Experimental/AVCodecs/NVCodecs/Source/NVENC/Private/Video/Encoders/VideoEncoderNVENC.cpp
--- a/Engine/Plugins/Experimental/AVCodecs/NVCodecs/Source/NVENC/Private/Video/Encoders/VideoEncoderNVENC.cpp
+++ b/Engine/Plugins/Experimental/AVCodecs/NVCodecs/Source/NVENC/Private/Video/Encoders/VideoEncoderNVENC.cpp
@@ -219,10 +219,6 @@ FAVResult FEncoderNVENC::SendFrame(TSharedPtr<FVideoResource> const& Resource, u
 			{
 				Picture.encodePicFlags |= NV_ENC_PIC_FLAG_FORCEIDR;
 			}
-			else // HACK (aidan.possemiers) AVCodecs h265 decoder only currently supports intra frames
-			{
-				Picture.encodePicFlags |= NV_ENC_PIC_FLAG_FORCEINTRA;
-			}
 		}
 		else
 		{
