Fix: Remove memory leak in FVideoResourceCUDA

#10637

diff --git a/Engine/Plugins/Experimental/AVCodecs/NVCodecs/Source/NVCodecs/Private/Video/Resources/VideoResourceCUDA.cpp b/Engine/Plugins/Experimental/AVCodecs/NVCodecs/Source/NVCodecs/Private/Video/Resources/VideoResourceCUDA.cpp
--- a/Engine/Plugins/Experimental/AVCodecs/NVCodecs/Source/NVCodecs/Private/Video/Resources/VideoResourceCUDA.cpp
+++ b/Engine/Plugins/Experimental/AVCodecs/NVCodecs/Source/NVCodecs/Private/Video/Resources/VideoResourceCUDA.cpp
@@ -5,7 +5,7 @@
 #include "Video/Resources/VideoResourceVulkan.h"
 
 #if PLATFORM_WINDOWS
-#include "Video/Resources/Windows/VideoResourceD3D.h"
+	#include "Video/Resources/Windows/VideoResourceD3D.h"
 #endif
 
 REGISTER_TYPEID(FVideoContextCUDA);
@@ -15,13 +15,13 @@ TAVResult<CUarray_format> ConvertFormat(EVideoFormat Format)
 {
 	switch (Format)
 	{
-	case EVideoFormat::R8:
-	case EVideoFormat::BGRA:
-		return CU_AD_FORMAT_UNSIGNED_INT8;
-	case EVideoFormat::G16:
-		return CU_AD_FORMAT_UNSIGNED_INT16;
-	default:
-		return FAVResult(EAVResult::ErrorUnsupported, FString::Printf(TEXT("EVideoFormat format %d is not supported"), Format), TEXT("CUDA"));
+		case EVideoFormat::R8:
+		case EVideoFormat::BGRA:
+			return CU_AD_FORMAT_UNSIGNED_INT8;
+		case EVideoFormat::G16:
+			return CU_AD_FORMAT_UNSIGNED_INT16;
+		default:
+			return FAVResult(EAVResult::ErrorUnsupported, FString::Printf(TEXT("EVideoFormat format %d is not supported"), Format), TEXT("CUDA"));
 	}
 }
 
@@ -134,6 +134,15 @@ FVideoResourceCUDA::~FVideoResourceCUDA()
 			FAVResult::Log(EAVResult::ErrorUnmapping, TEXT("Failed to clean up external memory"), TEXT("CUDA"), Result);
 		}
 	}
+
+	if (ExternalSemaphore)
+	{
+		CUresult const Result = FCUDAModule::CUDA().cuDestroyExternalSemaphore(ExternalSemaphore);
+		if (Result != CUDA_SUCCESS)
+		{
+			FAVResult::Log(EAVResult::ErrorUnmapping, TEXT("Failed to destroy external semaphore"), TEXT("CUDA"), Result);
+		}
+	}
 }
 
 FAVResult FVideoResourceCUDA::Validate() const
@@ -213,7 +222,7 @@ FAVResult FVideoResourceCUDA::CopyFromAsync(CUdeviceptr Source, uint32 MapPitch)
 		CUDA_EXTERNAL_SEMAPHORE_SIGNAL_PARAMS params = {};
 		params.params.fence.value = SemaphoreValue + 1;
 
-		FCUDAModule::CUDA().cuSignalExternalSemaphoresAsync(&ExternalSemaphore, &params, 1, cuStream);		
+		FCUDAModule::CUDA().cuSignalExternalSemaphoresAsync(&ExternalSemaphore, &params, 1, cuStream);
 	}
 
 	// Destroy the stream this will happen only when all tasks on the stream have completed
@@ -269,7 +278,7 @@ FAVResult FVideoResourceCUDA::CopyToAsync(CUdeviceptr Target, uint32 MapPitch)
 		CUDA_EXTERNAL_SEMAPHORE_SIGNAL_PARAMS params = {};
 		params.params.fence.value = SemaphoreValue + 1;
 
-		FCUDAModule::CUDA().cuSignalExternalSemaphoresAsync(&ExternalSemaphore, &params, 1, cuStream);		
+		FCUDAModule::CUDA().cuSignalExternalSemaphoresAsync(&ExternalSemaphore, &params, 1, cuStream);
 	}
 
 	// Destroy the stream this will happen only when all tasks on the stream have completed
@@ -321,7 +330,7 @@ DLLEXPORT FAVResult FAVExtension::TransformResource(TSharedPtr<FVideoResourceCUD
 		if (InResource->GetDevice()->HasContext<FVideoContextCUDA>())
 		{
 			// Depending on how we are running the engine the texture we are looking to bind could be
-			// stored as part of a heap or as a commited resource 
+			// stored as part of a heap or as a commited resource
 			CUDA_EXTERNAL_MEMORY_HANDLE_DESC ExternalResouceDesc = {};
 			if (InResource->GetHeap().IsValid())
 			{
@@ -346,7 +355,7 @@ DLLEXPORT FAVResult FAVExtension::TransformResource(TSharedPtr<FVideoResourceCUD
 			ExternalFenceDesc.handle.win32.handle = InResource->GetFenceSharedHandle();
 
 			OutResource = MakeShared<FVideoResourceCUDA>(InResource->GetDevice(), ExternalResouceDesc, ExternalFenceDesc, InResource.ToSharedRef(), InResource->GetFenceValue());
-			
+
 			return OutResource->Validate();
 		}
 
