Add: VideoToolbox Encoder and Decoders to AVCodecs

diff --git a/Engine/Plugins/Experimental/AVCodecs/NVCodecs/Source/NVDEC/Private/Video/Decoders/Configs/VideoDecoderConfigNVDEC.cpp b/Engine/Plugins/Experimental/AVCodecs/NVCodecs/Source/NVDEC/Private/Video/Decoders/Configs/VideoDecoderConfigNVDEC.cpp
--- a/Engine/Plugins/Experimental/AVCodecs/NVCodecs/Source/NVDEC/Private/Video/Decoders/Configs/VideoDecoderConfigNVDEC.cpp
+++ b/Engine/Plugins/Experimental/AVCodecs/NVCodecs/Source/NVDEC/Private/Video/Decoders/Configs/VideoDecoderConfigNVDEC.cpp
@@ -4,539 +4,6 @@
 
 REGISTER_TYPEID(FVideoDecoderConfigNVDEC);
 
-FAVResult FVideoDecoderConfigNVDEC::Parse(TSharedRef<FAVInstance> const& Instance, FVideoPacket const& Packet, TArray<FParsedPicture>& OutPictures)
-{
-	FAVResult Result = EAVResult::Success;
-
-	switch (CodecType)
-	{
-	case cudaVideoCodec_H264:
-		{
-			using namespace UE::AVCodecCore::H264;
-
-			FVideoDecoderConfigH264& H264 = Instance->Edit<FVideoDecoderConfigH264>();
-
-			TArray<Slice_t> Slices;
-			Result = H264.Parse(Instance, Packet, Slices);
-			if (Result.IsSuccess())
-			{
-				for (Slice_t const& Slice : Slices)
-				{
-					FParsedPicture& OutPicture = OutPictures.AddDefaulted_GetRef();
-					
-				
-					OutPicture.DecodeCreateInfo = *this;
-
-					CUVIDH264PICPARAMS& OutPictureH264 = OutPicture.CodecSpecific.h264;
-
-					// PPS
-					if (PPS_t const* const PPS = H264.PPS.Find(Slice.pic_parameter_set_id))
-					{
-						// SPS
-						if (SPS_t const* const SPS = H264.SPS.Find(PPS->seq_parameter_set_id))
-						{
-							OutPicture.DecodeCreateInfo.ulWidth = OutPicture.DecodeCreateInfo.ulTargetWidth = (SPS->pic_width_in_mbs_minus1 + 1) * 16;
-							OutPicture.DecodeCreateInfo.ulHeight = OutPicture.DecodeCreateInfo.ulTargetHeight = (SPS->pic_height_in_map_units_minus1 + 1) * 16;
-							OutPicture.DecodeCreateInfo.ChromaFormat = static_cast<cudaVideoChromaFormat>(SPS->chroma_format_idc.Value);
-							OutPicture.DecodeCreateInfo.ulIntraDecodeOnly = static_cast<unsigned long>(SPS->constraint_flags.Value) & EH264ConstraintFlag::Set3;
-						
-							OutPictureH264.entropy_coding_mode_flag = PPS->entropy_coding_mode_flag;
-							OutPictureH264.pic_order_present_flag = PPS->bottom_field_pic_order_in_frame_present_flag;
-							OutPictureH264.num_ref_idx_l0_active_minus1 = PPS->num_ref_idx_l0_default_active_minus1;
-							OutPictureH264.num_ref_idx_l1_active_minus1 = PPS->num_ref_idx_l1_default_active_minus1;
-							OutPictureH264.weighted_pred_flag = PPS->weighted_pred_flag;
-							OutPictureH264.weighted_bipred_idc = PPS->weighted_bipred_idc;
-							OutPictureH264.pic_init_qp_minus26 = PPS->pic_init_qp_minus26;
-							OutPictureH264.deblocking_filter_control_present_flag = PPS->deblocking_filter_control_present_flag;
-							OutPictureH264.redundant_pic_cnt_present_flag = PPS->redundant_pic_cnt_present_flag;
-							OutPictureH264.transform_8x8_mode_flag = PPS->transform_8x8_mode_flag;
-							OutPictureH264.MbaffFrameFlag = 0; // HACK does not seem to align directly with anything in spec
-							OutPictureH264.constrained_intra_pred_flag = PPS->constrained_intra_pred_flag;
-							OutPictureH264.chroma_qp_index_offset = PPS->chroma_qp_index_offset;
-							OutPictureH264.second_chroma_qp_index_offset = PPS->second_chroma_qp_index_offset;
-						
-							OutPictureH264.log2_max_frame_num_minus4 = SPS->log2_max_frame_num_minus4;
-							OutPictureH264.pic_order_cnt_type = SPS->pic_order_cnt_type;
-							OutPictureH264.log2_max_pic_order_cnt_lsb_minus4 = SPS->log2_max_pic_order_cnt_lsb_minus4;
-							OutPictureH264.delta_pic_order_always_zero_flag = SPS->delta_pic_order_always_zero_flag;
-							OutPictureH264.frame_mbs_only_flag = SPS->frame_mbs_only_flag;
-							OutPictureH264.direct_8x8_inference_flag = SPS->direct_8x8_inference_flag;
-							OutPictureH264.num_ref_frames = SPS->num_ref_frames_in_pic_order_cnt_cycle;
-							OutPictureH264.residual_colour_transform_flag = 0; // HACK does not seem to align directly with anything in spec
-							OutPictureH264.bit_depth_luma_minus8 = SPS->bit_depth_luma_minus8;
-							OutPictureH264.bit_depth_chroma_minus8 = SPS->bit_depth_chroma_minus8;
-							OutPictureH264.qpprime_y_zero_transform_bypass_flag = SPS->qpprime_y_zero_transform_bypass_flag;
-
-							// TODO (aidan) Other stuff? Not directly parsed from something must be inferred
-							OutPictureH264.ref_pic_flag = 0;
-							OutPictureH264.frame_num = 0;
-							OutPictureH264.CurrFieldOrderCnt[0] = 0;
-							OutPictureH264.CurrFieldOrderCnt[1] = 0;
-
-							// DPB
-							// TODO (aidan) this is complex as not directly parsed from something and must be inferred 
-							CUVIDH264DPBENTRY* DPB = OutPictureH264.dpb;
-							for (uint8 i = 0; i < 16u; i++)
-							{
-								// HACK (aidan)
-								static uint64 frameID = 0;
-					
-								DPB[0].PicIdx = frameID++;
-								DPB[0].used_for_reference = 3;
-
-								// TODO (aidan) handle case where we have a long term enabled  
-								DPB[0].is_long_term = 0;
-								DPB[0].FrameIdx = Slice.frame_num;
-
-								DPB[0].FieldOrderCnt[0] = 0;
-								DPB[0].FieldOrderCnt[1] = 0;
-							}
-
-							uint8 iYCbCr = SPS->separate_colour_plane_flag ? 1 : 0; // HACK (aidan) this comes from slice data that we are not yet parsing
-							bool bMbIsInterFlag = true;										   // HACK (aidan) part of the slice
-
-							if (PPS->pic_scaling_matrix_present_flag)
-							{
-								// Quantization Matrices
-								for (uint8 i = 0; i < 12u; i++)
-								{
-									// invoking 8.5.6 with ScalingList4x4[ iYCbCr + ( (mbIsInterFlag == 1 ) ? 3 : 0 )]
-									if (i < 6) // 4x4
-										{
-										ScaleListToWeightScale(true, PPS->ScalingList4x4[iYCbCr + (bMbIsInterFlag ? 3 : 0)], 16, OutPictureH264.WeightScale4x4[i]);
-										}
-									else // 8x8
-										{
-										ScaleListToWeightScale(true, PPS->ScalingList8x8[iYCbCr + (bMbIsInterFlag ? 3 : 0)], 64, OutPictureH264.WeightScale8x8[i]);
-										}
-								}
-							}
-							else if (SPS->seq_scaling_matrix_present_flag)
-							{
-								// Quantization Matrices
-								for (uint8 i = 0; i < 12u; i++)
-								{
-									// invoking 8.5.6 with ScalingList4x4[ iYCbCr + ( (mbIsInterFlag == 1 ) ? 3 : 0 )]
-									if (i < 6) // 4x4
-										{
-										ScaleListToWeightScale(true, SPS->ScalingList4x4[iYCbCr + (bMbIsInterFlag ? 3 : 0)], 16, OutPictureH264.WeightScale4x4[i]);
-										}
-									else // 8x8
-										{
-										ScaleListToWeightScale(true, SPS->ScalingList8x8[iYCbCr + (bMbIsInterFlag ? 3 : 0)], 64, OutPictureH264.WeightScale8x8[i]);
-										}
-								}
-							}
-						}
-						else
-						{
-							return FAVResult(EAVResult::ErrorInvalidState, FString::Printf(TEXT("SPS %d not found"), PPS->seq_parameter_set_id.Value), TEXT("H264"));
-						}
-					}
-					else
-					{
-						return FAVResult(EAVResult::ErrorInvalidState, FString::Printf(TEXT("PPS %d not found"), Slice.pic_parameter_set_id.Value), TEXT("H264"));
-					}
-				}
-			}
-		}
-
-		break;
-	case cudaVideoCodec_HEVC:
-		{
-			using namespace UE::AVCodecCore::H265;
-
-			FVideoDecoderConfigH265& H265 = Instance->Edit<FVideoDecoderConfigH265>();
-
-			TArray<TSharedPtr<FNaluH265>> Nalus;
-			Result = H265.Parse(Instance, Packet, Nalus);
-			if (Result.IsNotSuccess())
-			{
-				return FAVResult(EAVResult::ErrorResolving, TEXT("Failed to parse H265 Bitstream"));
-			}
-
-			// Loop through all the received NALUs and and add slices to the a FParsedPicture
-			// Depending on how we are receiving the bitstream we might get more than one picture
-			// so we handle that case by making multiple FParsedPictures and queue them
-			int32 index = 0;
-			while (index < Nalus.Num())
-			{
-				FParsedPicture& OutPicture = OutPictures.AddZeroed_GetRef();
-				
-				// NVDEC expects the bitstream to start with the first slice NAL delimiter at 0
-				OutPicture.pBitstreamData = Nalus[index]->StartCodeSize == 3 ? Nalus[index]->EBSP : Nalus[index]->EBSP + 1;
-
-				bool bParsedFirstSlice = false;
-
-				// TODO (aidan) this loop needs a clean-up wrt handling non-slices or non-first-segment slices
-				for (; index < Nalus.Num() && Nalus[index]->IsSlice(); index++)
-				{
-					TSharedPtr<FNaluSlice> Slice = StaticCastSharedPtr<FNaluSlice>(Nalus[index]);
-
-					// Apparently we got a new first slice without some other kind of
-					// NALU in between so we are done with the current picture
-					if (bParsedFirstSlice && Slice->first_slice_segment_in_pic_flag)
-					{
-						break;
-					}
-
-					// First slice in set
-					if (Slice->first_slice_segment_in_pic_flag == 1)
-					{
-						CUVIDHEVCPICPARAMS& OutPictureH265 = OutPicture.CodecSpecific.hevc;
-						
-						// PPS
-						if (!H265.ParsedPPS.Contains(Slice->slice_pic_parameter_set_id))
-						{
-							return FAVResult(EAVResult::Error, TEXT("PPS for slice missing from Instance"));
-						}
-
-						TSharedPtr<FNaluPPS> PPS = H265.ParsedPPS.FindRef(Slice->slice_pic_parameter_set_id);
-
-						// SPS
-						if (!H265.ParsedSPS.Contains(PPS->pps_seq_parameter_set_id))
-						{
-							return FAVResult(EAVResult::Error, TEXT("SPS for slice missing from Instance"));
-						}
-
-						TSharedPtr<FNaluSPS> SPS = H265.ParsedSPS.FindRef(PPS->pps_seq_parameter_set_id);
-
-						// VPS
-						if (!H265.ParsedVPS.Contains(SPS->sps_video_parameter_set_id))
-						{
-							return FAVResult(EAVResult::Error, TEXT("VPS for slice missing from Instance"));
-						}
-
-						TSharedPtr<FNaluVPS> VPS = H265.ParsedVPS.FindRef(SPS->sps_video_parameter_set_id);
-												
-						H265.UpdateRPS(Slice.ToSharedRef());
-												
-						OutPicture.DecodeCreateInfo = *this;
-
-						CodecType = OutPicture.DecodeCreateInfo.CodecType = cudaVideoCodec_HEVC;
-						
-						OutPicture.CurrPicIdx = H265.CurrPicIdx;
-						
-						OutPicture.field_pic_flag = 0;
-						OutPicture.ref_pic_flag = 1;
-						OutPicture.intra_pic_flag = Slice->IsIDR() || PPS->constrained_intra_pred_flag || true; // HACK (aidan.possemiers) We only currently send intra frames
-
-						if (OutPicture.intra_pic_flag) // HACK for intra only frames
-						{
-							FMemory::Memset(RefPicIdx, -1, sizeof(RefPicIdx));
-						}
-
-						OutPicture.PicWidthInMbs = SPS->pic_width_in_luma_samples / 16;
-						OutPicture.FrameHeightInMbs = SPS->pic_height_in_luma_samples / 16;
-
-						ulNumDecodeSurfaces = OutPicture.DecodeCreateInfo.ulNumDecodeSurfaces = 8; //SPS->MaxDpbSize; //HACK (aidan.possemiers) we should get this from the bitstream
-						ulNumOutputSurfaces = OutPicture.DecodeCreateInfo.ulNumOutputSurfaces = 1;
-
-						bitDepthMinus8 = OutPicture.DecodeCreateInfo.bitDepthMinus8 = SPS->bit_depth_luma_minus8;
-						ulWidth = OutPicture.DecodeCreateInfo.ulWidth = SPS->pic_width_in_luma_samples;
-						ulHeight = OutPicture.DecodeCreateInfo.ulHeight = SPS->pic_height_in_luma_samples;
-						ChromaFormat = OutPicture.DecodeCreateInfo.ChromaFormat = static_cast<cudaVideoChromaFormat>(SPS->chroma_format_idc.Value);
-						ulIntraDecodeOnly = OutPicture.DecodeCreateInfo.ulIntraDecodeOnly = static_cast<unsigned long>(SPS->profile_tier_level.general_constraint_flags & EH265ConstraintFlag::intra_constraint_flag);
-
-						ulTargetWidth = OutPicture.DecodeCreateInfo.ulTargetWidth = SPS->pic_width_in_luma_samples;
-						ulTargetHeight = OutPicture.DecodeCreateInfo.ulTargetHeight = SPS->pic_height_in_luma_samples;
-
-						// Need to scale the out frame depending on the chroma format
-						if (SPS->conformance_window_flag)
-						{
-							display_area = OutPicture.DecodeCreateInfo.display_area = {
-								static_cast<short>(SPS->conf_win.left_offset * (SPS->chroma_format_idc == 1 || SPS->chroma_format_idc == 2 ? 2 : 1)),
-								static_cast<short>(SPS->conf_win.top_offset * (SPS->chroma_format_idc == 1 ? 2 : 1)),
-								static_cast<short>(SPS->conf_win.right_offset * (SPS->chroma_format_idc == 1 || SPS->chroma_format_idc == 2 ? 2 : 1)),
-								static_cast<short>(SPS->conf_win.bottom_offset * (SPS->chroma_format_idc == 1 ? 2 : 1))
-							};
-
-							ulTargetWidth = OutPicture.DecodeCreateInfo.ulTargetWidth -= OutPicture.DecodeCreateInfo.display_area.left + OutPicture.DecodeCreateInfo.display_area.right;
-							ulTargetHeight = OutPicture.DecodeCreateInfo.ulTargetHeight -= OutPicture.DecodeCreateInfo.display_area.top + OutPicture.DecodeCreateInfo.display_area.bottom;
-						}
-
-						// NVDEC only supports 8 or 16 bit output surfaces in 420 or 444 format
-						if (ChromaFormat == cudaVideoChromaFormat_420)
-						{
-							OutputFormat = OutPicture.DecodeCreateInfo.OutputFormat = (bitDepthMinus8 == 0 ? cudaVideoSurfaceFormat_NV12 : cudaVideoSurfaceFormat_P016);
-						}
-						else
-						{
-							OutputFormat = OutPicture.DecodeCreateInfo.OutputFormat = (bitDepthMinus8 == 0 ? cudaVideoSurfaceFormat_YUV444 : cudaVideoSurfaceFormat_YUV444_16Bit);
-						}
-
-						OutPictureH265.pic_width_in_luma_samples = SPS->pic_width_in_luma_samples;
-						OutPictureH265.pic_height_in_luma_samples = SPS->pic_height_in_luma_samples;
-						OutPictureH265.log2_min_luma_coding_block_size_minus3 = SPS->log2_min_luma_coding_block_size_minus3;
-						OutPictureH265.log2_diff_max_min_luma_coding_block_size = SPS->log2_diff_max_min_luma_coding_block_size;
-						OutPictureH265.log2_min_transform_block_size_minus2 = SPS->log2_min_luma_transform_block_size_minus2;
-						OutPictureH265.log2_diff_max_min_transform_block_size = SPS->log2_diff_max_min_luma_transform_block_size;
-						OutPictureH265.pcm_enabled_flag = SPS->pcm_enabled_flag;
-						OutPictureH265.log2_min_pcm_luma_coding_block_size_minus3 = SPS->log2_min_pcm_luma_coding_block_size_minus3;
-						OutPictureH265.log2_diff_max_min_pcm_luma_coding_block_size = SPS->log2_diff_max_min_pcm_luma_coding_block_size;
-						OutPictureH265.pcm_sample_bit_depth_luma_minus1 = SPS->pcm_sample_bit_depth_luma_minus1;
-
-						OutPictureH265.pcm_sample_bit_depth_chroma_minus1 = SPS->pcm_sample_bit_depth_chroma_minus1;
-						OutPictureH265.pcm_loop_filter_disabled_flag = SPS->pcm_loop_filter_disabled_flag;
-						OutPictureH265.strong_intra_smoothing_enabled_flag = SPS->strong_intra_smoothing_enabled_flag;
-						OutPictureH265.max_transform_hierarchy_depth_intra = SPS->max_transform_hierarchy_depth_intra;
-						OutPictureH265.max_transform_hierarchy_depth_inter = SPS->max_transform_hierarchy_depth_inter;
-						OutPictureH265.amp_enabled_flag = SPS->amp_enabled_flag;
-						OutPictureH265.separate_colour_plane_flag = SPS->separate_colour_plane_flag;
-						OutPictureH265.log2_max_pic_order_cnt_lsb_minus4 = SPS->log2_max_pic_order_cnt_lsb_minus4;
-
-						OutPictureH265.num_short_term_ref_pic_sets = SPS->num_short_term_ref_pic_sets;
-						OutPictureH265.long_term_ref_pics_present_flag = SPS->long_term_ref_pics_present_flag;
-						OutPictureH265.num_long_term_ref_pics_sps = SPS->num_long_term_ref_pics_sps;
-						OutPictureH265.sps_temporal_mvp_enabled_flag = SPS->sps_temporal_mvp_enabled_flag;
-						OutPictureH265.sample_adaptive_offset_enabled_flag = SPS->sample_adaptive_offset_enabled_flag;
-						OutPictureH265.scaling_list_enable_flag = SPS->scaling_list_enabled_flag;
-						OutPictureH265.IrapPicFlag = (uint8)Slice->nal_unit_type.Value >= 16 && (uint8)Slice->nal_unit_type.Value <= 23;
-						OutPictureH265.IdrPicFlag = Slice->nal_unit_type == ENaluType::IDR_N_LP || Slice->nal_unit_type == ENaluType::IDR_W_RADL;
-
-						OutPictureH265.bit_depth_luma_minus8 = SPS->bit_depth_luma_minus8;
-						OutPictureH265.bit_depth_chroma_minus8 = SPS->bit_depth_chroma_minus8;
-
-						if (SPS->sps_range_extension_flag)
-						{
-							OutPictureH265.sps_range_extension_flag = SPS->sps_range_extension_flag;
-							OutPictureH265.high_precision_offsets_enabled_flag = SPS->sps_range_extension.high_precision_offsets_enabled_flag;
-							OutPictureH265.transform_skip_rotation_enabled_flag = SPS->sps_range_extension.transform_skip_rotation_enabled_flag;
-							OutPictureH265.transform_skip_context_enabled_flag = SPS->sps_range_extension.transform_skip_context_enabled_flag;
-							OutPictureH265.implicit_rdpcm_enabled_flag = SPS->sps_range_extension.implicit_rdpcm_enabled_flag;
-							OutPictureH265.explicit_rdpcm_enabled_flag = SPS->sps_range_extension.explicit_rdpcm_enabled_flag;
-							OutPictureH265.extended_precision_processing_flag = SPS->sps_range_extension.extended_precision_processing_flag;
-							OutPictureH265.intra_smoothing_disabled_flag = SPS->sps_range_extension.intra_smoothing_disabled_flag;
-							OutPictureH265.persistent_rice_adaptation_enabled_flag = SPS->sps_range_extension.persistent_rice_adaptation_enabled_flag;
-							OutPictureH265.cabac_bypass_alignment_enabled_flag = SPS->sps_range_extension.cabac_bypass_alignment_enabled_flag;
-						}
-
-						if (PPS->pps_range_extension_flag)
-						{
-							OutPictureH265.pps_range_extension_flag = PPS->pps_range_extension_flag;
-							OutPictureH265.log2_max_transform_skip_block_size_minus2 = PPS->pps_range_extension.log2_max_transform_skip_block_size_minus2;
-							OutPictureH265.log2_sao_offset_scale_luma = PPS->pps_range_extension.log2_sao_offset_scale_chroma;
-							OutPictureH265.log2_sao_offset_scale_chroma = PPS->pps_range_extension.log2_sao_offset_scale_chroma;
-							OutPictureH265.cross_component_prediction_enabled_flag = PPS->pps_range_extension.cross_component_prediction_enabled_flag;
-							OutPictureH265.chroma_qp_offset_list_enabled_flag = PPS->pps_range_extension.chroma_qp_offset_list_enabled_flag;
-							OutPictureH265.diff_cu_chroma_qp_offset_depth = PPS->pps_range_extension.diff_cu_chroma_qp_offset_depth;
-							OutPictureH265.chroma_qp_offset_list_len_minus1 = PPS->pps_range_extension.chroma_qp_offset_list_len_minus1;
-
-							FMemory::Memzero(OutPictureH265.cb_qp_offset_list, 21);
-							FMemory::Memcpy(OutPictureH265.cb_qp_offset_list, PPS->pps_range_extension.cb_qp_offset_list.GetData(), PPS->pps_range_extension.cb_qp_offset_list.Num());
-
-							FMemory::Memzero(OutPictureH265.cr_qp_offset_list, 21);
-							FMemory::Memcpy(OutPictureH265.cr_qp_offset_list, PPS->pps_range_extension.cr_qp_offset_list.GetData(), PPS->pps_range_extension.cr_qp_offset_list.Num());
-						}
-
-						OutPictureH265.dependent_slice_segments_enabled_flag = PPS->dependent_slice_segments_enabled_flag;
-						OutPictureH265.slice_segment_header_extension_present_flag = PPS->slice_segment_header_extension_present_flag;
-						OutPictureH265.sign_data_hiding_enabled_flag = PPS->sign_data_hiding_enabled_flag;
-						OutPictureH265.cu_qp_delta_enabled_flag = PPS->cu_qp_delta_enabled_flag;
-						OutPictureH265.diff_cu_qp_delta_depth = PPS->diff_cu_qp_delta_depth;
-						OutPictureH265.init_qp_minus26 = PPS->init_qp_minus26;
-						OutPictureH265.pps_cb_qp_offset = PPS->pps_cb_qp_offset;
-						OutPictureH265.pps_cr_qp_offset = PPS->pps_cr_qp_offset;
-
-						OutPictureH265.constrained_intra_pred_flag = PPS->constrained_intra_pred_flag;
-						OutPictureH265.weighted_pred_flag = PPS->weighted_pred_flag;
-						OutPictureH265.weighted_bipred_flag = PPS->weighted_bipred_flag;
-						OutPictureH265.transform_skip_enabled_flag = PPS->transform_skip_enabled_flag;
-						OutPictureH265.transquant_bypass_enabled_flag = PPS->transquant_bypass_enabled_flag;
-						OutPictureH265.entropy_coding_sync_enabled_flag = PPS->entropy_coding_sync_enabled_flag;
-						OutPictureH265.log2_parallel_merge_level_minus2 = PPS->log2_parallel_merge_level_minus2;
-						OutPictureH265.num_extra_slice_header_bits = PPS->num_extra_slice_header_bits;
-
-						OutPictureH265.loop_filter_across_tiles_enabled_flag = PPS->loop_filter_across_tiles_enabled_flag;
-						OutPictureH265.loop_filter_across_slices_enabled_flag = PPS->pps_loop_filter_across_slices_enabled_flag;
-						OutPictureH265.output_flag_present_flag = PPS->output_flag_present_flag;
-						OutPictureH265.num_ref_idx_l0_default_active_minus1 = PPS->num_ref_idx_l0_default_active_minus1;
-						OutPictureH265.num_ref_idx_l1_default_active_minus1 = PPS->num_ref_idx_l1_default_active_minus1;
-						OutPictureH265.lists_modification_present_flag = PPS->lists_modification_present_flag;
-						OutPictureH265.cabac_init_present_flag = PPS->cabac_init_present_flag;
-						OutPictureH265.pps_slice_chroma_qp_offsets_present_flag = PPS->pps_slice_chroma_qp_offsets_present_flag;
-
-						OutPictureH265.deblocking_filter_override_enabled_flag = PPS->deblocking_filter_override_enabled_flag;
-						OutPictureH265.pps_deblocking_filter_disabled_flag = PPS->pps_deblocking_filter_disabled_flag;
-						OutPictureH265.pps_beta_offset_div2 = PPS->pps_beta_offset_div2;
-						OutPictureH265.pps_tc_offset_div2 = PPS->pps_tc_offset_div2;
-						OutPictureH265.tiles_enabled_flag = PPS->tiles_enabled_flag;
-						OutPictureH265.uniform_spacing_flag = PPS->uniform_spacing_flag;
-						OutPictureH265.num_tile_columns_minus1 = PPS->num_tile_columns_minus1;
-						OutPictureH265.num_tile_rows_minus1 = PPS->num_tile_rows_minus1;
-
-						FMemory::Memzero(OutPictureH265.column_width_minus1, 21);
-						for (uint8 i = 0; i < PPS->column_width_minus1.Num(); i++)
-						{
-							OutPictureH265.column_width_minus1[i] = PPS->column_width_minus1[i];
-						}
-
-						FMemory::Memzero(OutPictureH265.row_height_minus1, 21);
-						for (uint8 i = 0; i < PPS->row_height_minus1.Num(); i++)
-						{
-							OutPictureH265.row_height_minus1[i] = PPS->row_height_minus1[i];
-						}
-
-						OutPictureH265.NumBitsForShortTermRPSInSlice = Slice->NumBitsForShortTermRPSInSlice;
-						OutPictureH265.NumDeltaPocsOfRefRpsIdx = SPS->short_term_ref_pic_sets[Slice->CurrRpsIdx].NumDeltaPocs;
-						OutPictureH265.NumPocTotalCurr = H265.ReferencePictureSet.NumPicTotalCurr;
-						OutPictureH265.NumPocStCurrBefore = H265.ReferencePictureSet.NumPocStCurrBefore;
-						OutPictureH265.NumPocStCurrAfter = H265.ReferencePictureSet.NumPocStCurrAfter;
-						OutPictureH265.NumPocLtCurr = H265.ReferencePictureSet.NumPocLtCurr;
-						OutPictureH265.CurrPicOrderCntVal = Slice->CurrPicOrderCntVal;
-
-						if (!OutPicture.intra_pic_flag)
-						{
-							for (uint8 i = 0; i < 16; i++)
-							{
-								if (i < SPS->MaxDpbSize)
-								{
-									OutPictureH265.RefPicIdx[i] = H265.ReferencePictureSet.PicUsage[RefPicIdx[i]] == EPictureUsage::UNUSED ? -1 : RefPicIdx[i];
-									OutPictureH265.IsLongTerm[i] = (H265.ReferencePictureSet.PicUsage[i] == EPictureUsage::LONG_TERM);
-								}
-								else
-								{
-									OutPictureH265.RefPicIdx[i] = -1;
-								}
-							}
-						}
-
-						
-						FMemory::Memcpy(OutPictureH265.PicOrderCntVal, H265.ReferencePictureSet.PicOrderCntVal, SPS->MaxDpbSize);
-						FMemory::Memcpy(OutPictureH265.RefPicSetStCurrBefore, H265.ReferencePictureSet.RefPicSetStCurrBefore, 8);
-						FMemory::Memcpy(OutPictureH265.RefPicSetStCurrAfter, H265.ReferencePictureSet.RefPicSetStCurrAfter, 8);
-						FMemory::Memcpy(OutPictureH265.RefPicSetLtCurr, H265.ReferencePictureSet.RefPicSetLtCurr, 8);
-
-
-
-						// HACK (aidan) we are not using layers yet so these should be zero until we do
-						FMemory::Memcpy(OutPictureH265.RefPicSetInterLayer0, H265.ReferencePictureSet.RefPicSetInterLayer0, 8);
-						FMemory::Memcpy(OutPictureH265.RefPicSetInterLayer1, H265.ReferencePictureSet.RefPicSetInterLayer1, 8);
-						
-						// TODO (aidan) check if you jump straight to default if the value is not present in the current active PPS or if you check SPS first
-						auto const& FillScalingLists = [&OutPictureH265](scaling_list_data_t const& ScalingListData) -> void {
-							for (uint8 i = 0; i < 6; i++)
-							{
-								if (ScalingListData.scaling_list_pred_mode_flag[0][i] > 0)
-								{
-									FMemory::Memcpy(OutPictureH265.ScalingList4x4[i], ScalingListData.ScalingList0[i], 16);
-								}
-								else
-								{
-									FMemory::Memcpy(OutPictureH265.ScalingList4x4[i], DefaultScalingList.ScalingList0[i], 16);
-								}
-							}
-
-							for (uint8 i = 0; i < 6; i++)
-							{
-								if (ScalingListData.scaling_list_pred_mode_flag[1][i] > 0)
-								{
-									FMemory::Memcpy(OutPictureH265.ScalingList8x8[i], ScalingListData.ScalingList1to3[0][i], 64);
-								}
-								else
-								{
-									FMemory::Memcpy(OutPictureH265.ScalingList8x8[i], DefaultScalingList.ScalingList1to3[0][i], 64);
-								}
-							}
-
-							for (uint8 i = 0; i < 6; i++)
-							{
-								if (ScalingListData.scaling_list_pred_mode_flag[2][i] > 0)
-								{
-									FMemory::Memcpy(OutPictureH265.ScalingList16x16[i], ScalingListData.ScalingList1to3[1][i], 64);
-									OutPictureH265.ScalingListDCCoeff16x16[i] = ScalingListData.scaling_list_dc_coef_minus8[2][i] + 8;
-								}
-								else
-								{
-									FMemory::Memcpy(OutPictureH265.ScalingList16x16[i], DefaultScalingList.ScalingList1to3[1][i], 64);
-									OutPictureH265.ScalingListDCCoeff16x16[i] = 16;
-								}
-							}
-
-							for (uint8 i = 0; i < 2; i++)
-							{
-								if (ScalingListData.scaling_list_pred_mode_flag[3][i] > 0)
-								{
-									FMemory::Memcpy(OutPictureH265.ScalingList32x32[i], ScalingListData.ScalingList1to3[2][i], 64);
-									OutPictureH265.ScalingListDCCoeff32x32[i] = ScalingListData.scaling_list_dc_coef_minus8[3][i] + 8;
-								}
-								else
-								{
-									FMemory::Memcpy(OutPictureH265.ScalingList32x32[i], DefaultScalingList.ScalingList1to3[2][i], 64);
-									OutPictureH265.ScalingListDCCoeff32x32[i] = 16;
-								}
-							}
-						};
-
-						if (PPS->pps_scaling_list_data_present_flag)
-						{
-							FillScalingLists(PPS->scaling_list_data);
-						}
-						else if (SPS->sps_scaling_list_data_present_flag)
-						{
-							FillScalingLists(SPS->scaling_list_data);
-						}
-						else
-						{
-							FillScalingLists(DefaultScalingList);
-						}
-						
-						OutPicture.SliceOffsets.Add(Slice->StartIdx);
-
-						// Could be a HACK (aidan) but reference NVDecoder parser seems to use this as a ring buffer
-						RefPicIdx[H265.CurrPicIdx % SPS->MaxDpbSize] = H265.CurrPicIdx % SPS->MaxDpbSize;
-						RefPicIdx[(H265.CurrPicIdx + 1) % SPS->MaxDpbSize] = -1;
-
-						bParsedFirstSlice = true;
-					}
-					else // Not first slice but member of current picture
-					{
-						OutPicture.SliceOffsets.Add(Slice->StartIdx); // NVDEC wants to see the 3 byte nal delimiter
-					}
-				}
-
-				if (OutPicture.SliceOffsets.Num() > 0)
-				{
-
-					// Make the offsets relative to the pointer to the first slice index in the bitstream
-					for (uint8 i = 1; i < OutPicture.SliceOffsets.Num(); i++)
-					{
-						OutPicture.SliceOffsets[i] -= OutPicture.SliceOffsets[0];
-					}
-
-					// Zero out the first slice index
-					OutPicture.SliceOffsets[0] = 0;
-
-					// Get size in bytes of the Bitstream with this PicParams slices					
-					OutPicture.nBitstreamDataLen = OutPicture.SliceOffsets.Last() - OutPicture.SliceOffsets[0] + Nalus[index - 1]->Size;
-
-					OutPicture.nNumSlices = OutPicture.SliceOffsets.Num();
-					OutPicture.pSliceDataOffsets = OutPicture.SliceOffsets.GetData();
-				}
-				else
-				{
-					OutPictures.Pop();
-				}
-				
-
-				// Loop through the remaining Nalus looking for the next slice
-				while (index < Nalus.Num() && !Nalus[index]->IsSlice())
-				{
-					index++;
-				}
-			} // while (index < Nalus.Num())
-
-			return EAVResult::Success;
-
-		} // case cudaVideoCodec_HEVC:
-		break;
-	default:
-		Result = FAVResult(EAVResult::ErrorMapping, FString::Printf(TEXT("Unsupported codec type %d"), CodecType), TEXT("NVDEC"));
-	}
-
-	return Result;
-}
-
 template <>
 DLLEXPORT FAVResult FAVExtension::TransformConfig(FVideoDecoderConfigNVDEC& OutConfig, FVideoDecoderConfig const& InConfig)
 {
diff --git a/Engine/Plugins/Experimental/AVCodecs/NVCodecs/Source/NVENC/Public/Video/Encoders/VideoEncoderNVENC.h b/Engine/Plugins/Experimental/AVCodecs/NVCodecs/Source/NVENC/Public/Video/Encoders/VideoEncoderNVENC.h
--- a/Engine/Plugins/Experimental/AVCodecs/NVCodecs/Source/NVENC/Public/Video/Encoders/VideoEncoderNVENC.h
+++ b/Engine/Plugins/Experimental/AVCodecs/NVCodecs/Source/NVENC/Public/Video/Encoders/VideoEncoderNVENC.h
@@ -11,7 +11,7 @@
 #include "Containers/Queue.h"
 
 #if PLATFORM_WINDOWS
-	#include "Video/Resources/Windows/VideoResourceD3D.h"
+	#include "Video/Resources/D3D/VideoResourceD3D.h"
 #endif // PLATFORM_WINDOWS
 
 /**
diff --git a/Engine/Plugins/Experimental/AVCodecs/VTCodecs/Source/VTCodecs/Private/VTCodecsModule.cpp b/Engine/Plugins/Experimental/AVCodecs/VTCodecs/Source/VTCodecs/Private/VTCodecsModule.cpp
new file mode 100644
--- /dev/null
+++ b/Engine/Plugins/Experimental/AVCodecs/VTCodecs/Source/VTCodecs/Private/VTCodecsModule.cpp
@@ -0,0 +1,68 @@
+// Copyright Epic Games, Inc. All Rights Reserved.
+
+#include "Video/Encoders/VideoEncoderVT.h"
+#include "Video/Encoders/Configs/VideoEncoderConfigH264.h"
+#include "Video/Encoders/Configs/VideoEncoderConfigH265.h"
+
+#include "Video/Decoders/VideoDecoderVT.h"
+#include "Video/Decoders/Configs/VideoDecoderConfigVP9.h"
+#include "Video/Decoders/Configs/VideoDecoderConfigH264.h"
+#include "Video/Decoders/Configs/VideoDecoderConfigH265.h"
+
+#include "Video/Resources/Metal/VideoResourceMetal.h"
+
+THIRD_PARTY_INCLUDES_START
+#include <VideoToolbox/VideoToolbox.h>
+THIRD_PARTY_INCLUDES_END
+
+class FVTModule : public IModuleInterface
+{
+public:
+	virtual void StartupModule() override
+	{       
+        VTRegisterSupplementalVideoDecoderIfAvailable(kCMVideoCodecType_VP9);
+
+        FVideoEncoder::RegisterPermutationsOf<TVideoEncoderVT<FVideoResourceMetal>>
+            ::With<FVideoResourceMetal>
+            ::And<FVideoEncoderConfigVT, FVideoEncoderConfigH264, FVideoEncoderConfigH265>(
+                [](TSharedRef<FAVDevice> const& NewDevice, TSharedRef<FAVInstance> const& NewInstance)
+                {
+                    return FAPI::Get<FVT>().IsValid();
+                });
+
+        if(VTIsHardwareDecodeSupported(kCMVideoCodecType_H264))
+        {
+            FVideoDecoder::RegisterPermutationsOf<TVideoDecoderVT<FVideoResourceMetal>>
+                ::With<FVideoResourceMetal>
+                ::And<FVideoDecoderConfigVT, FVideoDecoderConfigH264>(
+                    [](TSharedRef<FAVDevice> const& NewDevice, TSharedRef<FAVInstance> const& NewInstance)
+                    {
+                        return FAPI::Get<FVT>().IsValid();
+                    });
+        }
+
+        if(VTIsHardwareDecodeSupported(kCMVideoCodecType_HEVC))
+        {
+            FVideoDecoder::RegisterPermutationsOf<TVideoDecoderVT<FVideoResourceMetal>>
+                ::With<FVideoResourceMetal>
+                ::And<FVideoDecoderConfigVT, FVideoDecoderConfigH265>(
+                    [](TSharedRef<FAVDevice> const& NewDevice, TSharedRef<FAVInstance> const& NewInstance)
+                    {
+                        return FAPI::Get<FVT>().IsValid();
+                    });
+        }
+
+        if(VTIsHardwareDecodeSupported(kCMVideoCodecType_VP9))
+        {
+            FVideoDecoder::RegisterPermutationsOf<TVideoDecoderVT<FVideoResourceMetal>>
+                ::With<FVideoResourceMetal>
+                ::And<FVideoDecoderConfigVT, FVideoDecoderConfigVP9>(
+                    [](TSharedRef<FAVDevice> const& NewDevice, TSharedRef<FAVInstance> const& NewInstance)
+                    {
+                        return FAPI::Get<FVT>().IsValid();
+                    });
+        }
+	}
+};
+
+IMPLEMENT_MODULE(FVTModule, VT);
diff --git a/Engine/Plugins/Experimental/AVCodecs/VTCodecs/Source/VTCodecs/Public/Video/Decoders/VideoDecoderVT.hpp b/Engine/Plugins/Experimental/AVCodecs/VTCodecs/Source/VTCodecs/Public/Video/Decoders/VideoDecoderVT.hpp
new file mode 100644
--- /dev/null
+++ b/Engine/Plugins/Experimental/AVCodecs/VTCodecs/Source/VTCodecs/Public/Video/Decoders/VideoDecoderVT.hpp
@@ -0,0 +1,360 @@
+// Copyright Epic Games, Inc. All Rights Reserved.
+
+#pragma once
+
+#include "Video/Resources/Metal/VideoResourceMetal.h"
+#include "Video/Util/NaluRewriter.h"
+
+THIRD_PARTY_INCLUDES_START
+#include <VideoToolbox/VideoToolbox.h>
+#include <CoreMedia/CMSync.h>
+THIRD_PARTY_INCLUDES_END
+
+#define CONDITIONAL_RELEASE(x)          \
+    if (x)                              \
+    {                                   \
+        CFRelease(x);                   \
+        x = nullptr;                    \
+    }
+
+template <typename TResource>
+TVideoDecoderVT<TResource>::~TVideoDecoderVT()
+{
+	Close();
+}
+
+template <typename TResource>
+bool TVideoDecoderVT<TResource>::IsOpen() const
+{
+	return bIsOpen;
+}
+
+template <typename TResource>
+FAVResult TVideoDecoderVT<TResource>::Open(TSharedRef<FAVDevice> const& NewDevice, TSharedRef<FAVInstance> const& NewInstance)
+{
+	Close();
+
+	TVideoDecoder<TResource, FVideoDecoderConfigVT>::Open(NewDevice, NewInstance);
+
+    MemoryPool = CMMemoryPoolCreate(nullptr);
+
+	FrameCount = 0;
+
+    bIsOpen = true;
+
+	return EAVResult::Success;
+}
+
+template <typename TResource>
+void TVideoDecoderVT<TResource>::Close()
+{
+    DestroyDecompressionSession();
+    SetVideoFormat(nullptr);
+
+    if(MemoryPool)
+    {
+        CMMemoryPoolInvalidate(MemoryPool);
+        CFRelease(MemoryPool);
+    }
+
+    bIsOpen = false;
+}
+
+template <typename TResource>
+void TVideoDecoderVT<TResource>::ResetDecompressionSession()
+{
+    DestroyDecompressionSession();
+
+    if(!VideoFormat)
+    {
+        FAVResult::Log(EAVResult::PendingInput, TEXT("Waiting for VideoFormat"), TEXT("VT"));                   
+        return; 
+    }
+    
+    // Set source image buffer attributes. These attributes will be present on
+    // buffers retrieved from the decoder's pixel buffer pool.
+    const size_t AttributesSize = 3;
+    CFTypeRef Keys[AttributesSize] = 
+    {
+        kCVPixelBufferOpenGLCompatibilityKey,
+        kCVPixelBufferIOSurfacePropertiesKey,
+        kCVPixelBufferPixelFormatTypeKey
+    };
+
+    CFDictionaryRef IOSurfaceValue = CFDictionaryCreate(kCFAllocatorDefault, nullptr, nullptr, 0, &kCFTypeDictionaryKeyCallBacks, &kCFTypeDictionaryValueCallBacks);
+    // TODO (belchy06): This should support more than ARGB8 pixel format
+    int64_t PixelType = kCVPixelFormatType_32BGRA;
+    CFNumberRef PixelFormat = CFNumberCreate(nullptr, kCFNumberLongType, &PixelType);
+
+    CFTypeRef Values[AttributesSize] =
+    { 
+        kCFBooleanTrue, 
+        IOSurfaceValue, 
+        PixelFormat 
+    };
+
+    CFDictionaryRef Attributes = CFDictionaryCreate(kCFAllocatorDefault, Keys, Values, AttributesSize, &kCFTypeDictionaryKeyCallBacks, &kCFTypeDictionaryValueCallBacks);
+
+    CONDITIONAL_RELEASE(IOSurfaceValue);
+    CONDITIONAL_RELEASE(PixelFormat);
+
+    VTDecompressionOutputCallbackRecord Record = 
+    {
+        Internal::VTDecompressionOutputCallback, 
+        this
+    };
+                
+    OSStatus Result = VTDecompressionSessionCreate(kCFAllocatorDefault, VideoFormat, nullptr, Attributes, &Record, &Decoder);
+
+    if(Result != 0)
+    {
+        DestroyDecompressionSession();
+        FAVResult::Log(EAVResult::ErrorCreating, TEXT("Failed to create VTDecompressionSession"), TEXT("VT"), Result);
+    }
+
+    CONDITIONAL_RELEASE(Attributes);
+
+    ConfigureDecompressionSession();
+}
+
+template <typename TResource>
+void TVideoDecoderVT<TResource>::DestroyDecompressionSession()
+{
+    if (Decoder) 
+    {
+        VTDecompressionSessionInvalidate(Decoder);	
+        CFRelease(Decoder);
+        Decoder = nullptr;
+    }
+}
+
+template <typename TResource>
+void TVideoDecoderVT<TResource>::ConfigureDecompressionSession()
+{
+    VTSessionSetProperty(Decoder, kVTDecompressionPropertyKey_RealTime, kCFBooleanTrue);
+}
+
+template <typename TResource>
+bool TVideoDecoderVT<TResource>::IsInitialized() const
+{
+	return Decoder != nullptr;
+}
+
+template <typename TResource>
+FAVResult TVideoDecoderVT<TResource>::ApplyConfig()
+{
+	if (IsOpen())
+	{
+		FVideoDecoderConfigVT const& PendingConfig = this->GetPendingConfig();
+		if (this->AppliedConfig != PendingConfig)
+		{
+			if (IsInitialized())
+			{
+                if(this->AppliedConfig.Codec == PendingConfig.Codec)
+                {
+                    // TODO (belchy06): Reconfiguration
+                }
+                else
+                {
+                    if (Decoder) 
+                    {
+                        DestroyDecompressionSession();
+                        FAVResult::Log(EAVResult::Success, TEXT("Re-initializing decoding session"), TEXT("VT"));
+                    }
+                }
+			}
+
+			if (!IsInitialized())
+			{
+                ResetDecompressionSession();
+			}
+		}
+
+		return TVideoDecoder<TResource, FVideoDecoderConfigVT>::ApplyConfig();
+	}
+
+	return FAVResult(EAVResult::ErrorInvalidState, TEXT("Decoder not open"), TEXT("VT"));
+}
+
+template <typename TResource>
+FAVResult TVideoDecoderVT<TResource>::SendPacket(FVideoPacket const& Packet)
+{
+	if (IsOpen())
+	{
+        // We've received our first call to decode a frame, we can now parse the information from
+        // the bitstream, configure our config and initialize the session
+		if(!IsInitialized())
+        {
+            // Applying the config must be done before parsing the format as we need to know the codec
+            FAVResult AVResult = ApplyConfig();
+		    if (AVResult.IsNotSuccess())
+		    {
+    			return AVResult;
+		    }
+
+            CMVideoFormatDescriptionRef InputFormat = nullptr;
+            if(this->AppliedConfig.Codec == kCMVideoCodecType_H264)
+            {
+                InputFormat = NaluRewriter::CreateH264VideoFormatDescription(Packet.DataPtr.Get(), Packet.DataSize);
+            }
+            else if(this->AppliedConfig.Codec == kCMVideoCodecType_HEVC)
+            {
+                InputFormat = NaluRewriter::CreateH265VideoFormatDescription(Packet.DataPtr.Get(), Packet.DataSize);
+            }
+            else if (this->AppliedConfig.Codec == kCMVideoCodecType_VP9)
+            {
+                InputFormat = NaluRewriter::CreateVP9VideoFormatDescription(Packet.DataPtr.Get(), Packet.DataSize);
+            }
+            else
+            {
+                return FAVResult(EAVResult::Error, TEXT("Unsupported codec"), TEXT("VT"));
+            }
+
+            if(InputFormat)
+            {
+                if(!CMFormatDescriptionEqual(InputFormat, VideoFormat))
+                {
+                    SetVideoFormat(InputFormat);
+                    ResetDecompressionSession();
+                }
+            }
+
+            CONDITIONAL_RELEASE(InputFormat);
+        }
+
+        if(!VideoFormat)
+        {
+            return FAVResult(EAVResult::WarningInvalidState, TEXT("Missing video format. Frame with sps/pps required."), TEXT("VT"));
+        }
+
+        CMSampleBufferRef SampleBuffer = nullptr;
+        if(this->AppliedConfig.Codec == kCMVideoCodecType_H264)
+        {
+            if(!NaluRewriter::H264AnnexBBufferToCMSampleBuffer(Packet.DataPtr.Get(), Packet.DataSize, VideoFormat, &SampleBuffer, MemoryPool))
+            {
+                return FAVResult(EAVResult::Error, TEXT("Failed to get SampleBuffer"), TEXT("VT"));
+            }
+        }
+        else if(this->AppliedConfig.Codec == kCMVideoCodecType_HEVC)
+        {
+            if(!NaluRewriter::H265AnnexBBufferToCMSampleBuffer(Packet.DataPtr.Get(), Packet.DataSize, VideoFormat, &SampleBuffer, MemoryPool))
+            {
+                return FAVResult(EAVResult::Error, TEXT("Failed to get SampleBuffer"), TEXT("VT"));
+            }
+        }
+        else if (this->AppliedConfig.Codec == kCMVideoCodecType_VP9)
+        {
+            if(!NaluRewriter::VP9BufferToCMSampleBuffer(Packet.DataPtr.Get(), Packet.DataSize, VideoFormat, &SampleBuffer, MemoryPool))
+            {
+                return FAVResult(EAVResult::Error, TEXT("Failed to get SampleBuffer"), TEXT("VT"));
+            }
+        }
+        else
+        {
+            return FAVResult(EAVResult::Error, TEXT("Unsupported codec"), TEXT("VT"));
+        }
+        
+        if(SampleBuffer == nullptr)
+        {
+            return FAVResult(EAVResult::Error, TEXT("SampleBuffer is nullptr"), TEXT("VT"));
+        }
+
+        TUniquePtr<DecodeParams> FrameDecodeParams;
+        FrameDecodeParams.Reset(new DecodeParams());
+
+        OSStatus Result = VTDecompressionSessionDecodeFrame(Decoder, SampleBuffer, 0, FrameDecodeParams.Release(), nullptr);
+
+        CFRelease(SampleBuffer);
+
+        if(Result != 0)
+        {
+            return FAVResult(EAVResult::Error, TEXT("Failed to decode frame"), TEXT("VT"), Result);
+        }
+
+        return EAVResult::Success;
+	}
+
+	return FAVResult(EAVResult::ErrorInvalidState, TEXT("Decoder not open"), TEXT("VT"));
+}
+
+template <typename TResource>
+FAVResult TVideoDecoderVT<TResource>::ReceiveFrame(TResolvableVideoResource<TResource>& InOutResource)
+{
+    if(IsOpen())
+    {
+        if(TSharedPtr<FFrame> Frame = *Frames.Peek())
+        {
+            size_t Width = CVPixelBufferGetWidth(Frame->ImageBuffer);
+            size_t Height = CVPixelBufferGetHeight(Frame->ImageBuffer);
+
+            OSType PixelFormat = CVPixelBufferGetPixelFormatType(Frame->ImageBuffer);
+
+            FVideoDescriptor ResourceDescriptor = FVideoDescriptor(EVideoFormat::BGRA, Width, Height);
+            if (!InOutResource.Resolve(this->GetDevice(), ResourceDescriptor))
+			{
+				return FAVResult(EAVResult::ErrorResolving, TEXT("Failed to resolve frame resource"), TEXT("VT"));
+			}
+
+            CVPixelBufferLockBaseAddress(Frame->ImageBuffer, kCVPixelBufferLock_ReadOnly);
+
+            void* PixelPtr = CVPixelBufferGetBaseAddress(Frame->ImageBuffer);
+
+            mtlpp::Texture* RawTexture = InOutResource->GetRaw();
+
+            // Do copy into the VideoResource
+            InOutResource->GetRaw()->Replace(mtlpp::Region(0, 0, Width, Height), 0, PixelPtr, Width * 4);
+
+            CVPixelBufferUnlockBaseAddress(Frame->ImageBuffer, kCVPixelBufferLock_ReadOnly);
+
+            Frames.Pop();
+            
+            return EAVResult::Success;
+        }
+
+        return EAVResult::PendingInput;
+    }
+
+	return FAVResult(EAVResult::ErrorInvalidState, TEXT("Decoder not open"), TEXT("VT"));
+}
+
+template <typename TResource>
+void TVideoDecoderVT<TResource>::SetVideoFormat(CMVideoFormatDescriptionRef Format)
+{
+    if(VideoFormat == Format)
+    {
+        return;
+    }
+
+    CONDITIONAL_RELEASE(VideoFormat);
+
+    VideoFormat = Format;
+    if(VideoFormat)
+    {
+        CFRetain(VideoFormat);
+    }
+}
+
+template <typename TResource>
+FAVResult TVideoDecoderVT<TResource>::HandleFrame(void* Params, OSStatus Status, VTDecodeInfoFlags InfoFlags, CVImageBufferRef ImageBuffer, CMTime Timestamp, CMTime Duration)
+{
+    if(IsOpen())
+    {
+        if(Status != 0)
+        {
+            return FAVResult(EAVResult::Error, TEXT("Failed to decode"), TEXT("VT"), Status);
+        }
+
+        if(!ImageBuffer)
+        {
+            return FAVResult(EAVResult::Error, TEXT("No output image buffer"), TEXT("VT"), Status);
+        }
+        
+        // The destructor for FFrame releases the ImageBuffer so we need to make sure it's not release until after it's been used
+        TSharedPtr<FFrame> Frame = MakeShareable(new FFrame(ImageBuffer, Timestamp, Duration));
+        Frames.Enqueue(Frame);
+
+        return EAVResult::Success;
+    }
+    
+    return FAVResult(EAVResult::ErrorInvalidState, TEXT("Decoder not open"), TEXT("VT"));
+}
diff --git a/Engine/Plugins/Experimental/AVCodecs/VTCodecs/Source/VTCodecs/Public/Video/Encoders/VideoEncoderVT.hpp b/Engine/Plugins/Experimental/AVCodecs/VTCodecs/Source/VTCodecs/Public/Video/Encoders/VideoEncoderVT.hpp
new file mode 100644
--- /dev/null
+++ b/Engine/Plugins/Experimental/AVCodecs/VTCodecs/Source/VTCodecs/Public/Video/Encoders/VideoEncoderVT.hpp
@@ -0,0 +1,402 @@
+// Copyright Epic Games, Inc. All Rights Reserved.
+
+#pragma once
+
+#include "Video/Resources/Metal/VideoResourceMetal.h"
+#include "Video/Util/NaluRewriter.h"
+
+THIRD_PARTY_INCLUDES_START
+#include <VideoToolbox/VideoToolbox.h>
+#include <CoreMedia/CMSync.h>
+THIRD_PARTY_INCLUDES_END
+
+#define CONDITIONAL_RELEASE(x)          \
+    if (x)                              \
+    {                                   \
+        CFRelease(x);                   \
+        x = nullptr;                    \
+    }
+
+template <typename TResource>
+TVideoEncoderVT<TResource>::~TVideoEncoderVT()
+{
+	Close();
+}
+
+template <typename TResource>
+bool TVideoEncoderVT<TResource>::IsOpen() const
+{
+	return bIsOpen;
+}
+
+template <typename TResource>
+FAVResult TVideoEncoderVT<TResource>::Open(TSharedRef<FAVDevice> const& NewDevice, TSharedRef<FAVInstance> const& NewInstance)
+{
+	Close();
+
+	TVideoEncoder<TResource, FVideoEncoderConfigVT>::Open(NewDevice, NewInstance);
+
+	FrameCount = 0;
+
+    bIsOpen = true;
+
+	return EAVResult::Success;
+}
+
+template <typename TResource>
+void TVideoEncoderVT<TResource>::Close()
+{
+    if (Encoder) 
+    {
+        VTCompressionSessionInvalidate(Encoder);
+        CFRelease(Encoder);
+        Encoder = nullptr;
+    }
+}
+
+template <typename TResource>
+bool TVideoEncoderVT<TResource>::IsInitialized() const
+{
+	return Encoder != nullptr;
+}
+
+template <typename TResource>
+FAVResult TVideoEncoderVT<TResource>::ApplyConfig()
+{
+	if (IsOpen())
+	{
+		FVideoEncoderConfigVT const& PendingConfig = this->GetPendingConfig();
+		if (this->AppliedConfig != PendingConfig)
+		{
+			if (IsInitialized())
+			{
+                if(this->AppliedConfig.Width == PendingConfig.Width
+                    && this->AppliedConfig.Height == PendingConfig.Height
+                    && this->AppliedConfig.Codec == PendingConfig.Codec)
+                {
+                    // Only bitrates can be configured on the fly
+                    SetEncoderBitrate(PendingConfig);
+                }
+                else
+                {
+                    if (Encoder) 
+                    {
+                        VTCompressionSessionInvalidate(Encoder);
+                        CFRelease(Encoder);
+                        Encoder = nullptr;
+                        FAVResult::Log(EAVResult::Success, TEXT("Re-initializing encoding session"), TEXT("VT"));
+                    }
+                }
+			}
+
+			if (!IsInitialized())
+			{
+                // Set source image buffer attributes. These attributes will be present on
+                // buffers retrieved from the encoder's pixel buffer pool.
+                const size_t AttributesSize = 3;
+                CFTypeRef Keys[AttributesSize] = {
+                    kCVPixelBufferOpenGLCompatibilityKey,
+                    kCVPixelBufferIOSurfacePropertiesKey,
+                    kCVPixelBufferPixelFormatTypeKey
+                };
+
+                CFDictionaryRef IOSurfaceValue = CFDictionaryCreate(kCFAllocatorDefault, nullptr, nullptr, 0, &kCFTypeDictionaryKeyCallBacks, &kCFTypeDictionaryValueCallBacks);
+               
+                int64 PixelType = PendingConfig.GetCVPixelFormatType();
+                CFNumberRef PixelFormat = CFNumberCreate(nullptr, kCFNumberLongType, &PixelType);
+
+                CFTypeRef Values[AttributesSize] = {kCFBooleanTrue, IOSurfaceValue, PixelFormat};
+                CFDictionaryRef SourceAttributes = CFDictionaryCreate(kCFAllocatorDefault, Keys, Values, AttributesSize, &kCFTypeDictionaryKeyCallBacks, &kCFTypeDictionaryValueCallBacks);
+
+                CONDITIONAL_RELEASE(IOSurfaceValue);
+                CONDITIONAL_RELEASE(PixelFormat);
+
+                CFMutableDictionaryRef EncoderSpecification = CFDictionaryCreateMutable(kCFAllocatorDefault, 2, &kCFTypeDictionaryKeyCallBacks, &kCFTypeDictionaryValueCallBacks);
+                
+                if(PendingConfig.Codec == kCMVideoCodecType_H264)
+                {
+                    // TODO (belchy06): Not all encoders are real-time. This should be moved to the config
+                    CFDictionarySetValue(EncoderSpecification, kVTVideoEncoderSpecification_EnableLowLatencyRateControl, kCFBooleanTrue);
+                }
+                else if(PendingConfig.Codec == kCMVideoCodecType_HEVC)
+                {
+                    CFDictionarySetValue(EncoderSpecification, kVTVideoEncoderSpecification_EnableHardwareAcceleratedVideoEncoder, kCFBooleanTrue);
+                }
+                
+                OSStatus Result = VTCompressionSessionCreate(kCFAllocatorDefault, PendingConfig.Width, PendingConfig.Height, PendingConfig.Codec, EncoderSpecification, SourceAttributes, NULL /* Default Compressed Data Allocator */, Internal::VTCompressionOutputCallback, this, &Encoder);
+                    
+                if(Result != 0)
+                {
+                    FAVResult::Log(EAVResult::ErrorCreating, TEXT("Failed to create VTCompressionSession"), TEXT("VT"), Result);
+                }
+
+                CONDITIONAL_RELEASE(SourceAttributes);
+                CONDITIONAL_RELEASE(EncoderSpecification);
+
+                ConfigureCompressionSession(PendingConfig);
+			}
+		}
+
+		return TVideoEncoder<TResource, FVideoEncoderConfigVT>::ApplyConfig();
+	}
+
+	return FAVResult(EAVResult::ErrorInvalidState, TEXT("Encoder not open"), TEXT("VT"));
+}
+
+template <typename TResource>
+FAVResult TVideoEncoderVT<TResource>::ConfigureCompressionSession(FVideoEncoderConfigVT const& Config)  
+{
+    VTSessionHelpers::SetVTSessionProperty(Encoder, kVTCompressionPropertyKey_RealTime, true);
+    
+    if(Config.Codec == kCMVideoCodecType_H264)
+    {
+        VTSessionHelpers::SetVTSessionProperty(Encoder, kVTCompressionPropertyKey_ProfileLevel, Config.Profile);
+    }
+    
+    VTSessionHelpers::SetVTSessionProperty(Encoder, kVTCompressionPropertyKey_AllowFrameReordering, false);
+
+    VTSessionHelpers::SetVTSessionProperty(Encoder, kVTCompressionPropertyKey_MaxKeyFrameInterval, Config.KeyframeInterval);
+
+    VTSessionHelpers::SetVTSessionProperty(Encoder, kVTCompressionPropertyKey_ExpectedFrameRate, Config.FrameRate);
+
+    SetEncoderBitrate(Config);
+    
+    OSStatus Status = VTCompressionSessionPrepareToEncodeFrames(Encoder);
+    if(Status != 0)
+    {
+        return FAVResult(EAVResult::Warning, TEXT("Compression session failed to prepare encode frames"), TEXT("VT"));
+    }
+
+    return EAVResult::Success;
+}
+
+template <typename TResource>
+FAVResult TVideoEncoderVT<TResource>::SetEncoderBitrate(FVideoEncoderConfigVT const& Config)
+{
+    // CBR not supported by encoders in their current configuration
+    // if(Config.RateControlMode == ERateControlMode::CBR)
+    // {
+    //     VTSessionHelpers::SetVTSessionProperty(Encoder, kVTCompressionPropertyKey_ConstantBitRate, Config.TargetBitrate);
+    // }
+    // else
+    {
+        VTSessionHelpers::SetVTSessionProperty(Encoder, kVTCompressionPropertyKey_AverageBitRate, Config.TargetBitrate);
+
+        int64_t DataLimitBytesPerSecondValue = static_cast<int64_t>(Config.MaxBitrate * 1.5f / 8);
+
+        CFNumberRef BytesPerSecond = CFNumberCreate(kCFAllocatorDefault, kCFNumberSInt64Type, &DataLimitBytesPerSecondValue);
+        
+        int64_t OneSecondValue = 1;
+        CFNumberRef OneSecond = CFNumberCreate(kCFAllocatorDefault, kCFNumberSInt64Type, &OneSecondValue);
+    
+        const void* Nums[2] = { BytesPerSecond, OneSecond };
+    
+        CFArrayRef DataRateLimits = CFArrayCreate(nullptr, Nums, 2, &kCFTypeArrayCallBacks);
+    
+        OSStatus Result = VTSessionSetProperty(Encoder, kVTCompressionPropertyKey_DataRateLimits, DataRateLimits);
+
+        if(Result != 0)
+        {
+            FAVResult::Log(EAVResult::Error, TEXT("Failed to set kVTCompressionPropertyKey_DataRateLimits"), TEXT("VT"), Result);
+        }
+
+        CONDITIONAL_RELEASE(BytesPerSecond);
+        CONDITIONAL_RELEASE(OneSecond);
+        CONDITIONAL_RELEASE(DataRateLimits);
+    }
+
+    return EAVResult::Success;
+}
+
+template <typename TResource>
+FAVResult TVideoEncoderVT<TResource>::SendFrame(TSharedPtr<FVideoResourceMetal> const& Resource, uint32 Timestamp, bool bForceKeyframe)
+{
+    
+    if(IsOpen())
+    {
+        FVideoEncoderConfigVT& PendingConfig = this->EditPendingConfig();
+        PendingConfig.PixelFormat = Resource->GetFormat();
+
+        FAVResult AVResult = ApplyConfig();
+		if (AVResult.IsNotSuccess())
+		{
+			return AVResult;
+		}
+
+        if(Resource.IsValid())
+        {
+            // Get a pixel buffer from the pool and copy frame data over.
+            CVPixelBufferPoolRef PixelBufferPool = VTCompressionSessionGetPixelBufferPool(Encoder);
+            if(!PixelBufferPool)
+            {
+                return FAVResult(EAVResult::Error, TEXT("Failed to get pixel buffer pool"), TEXT("VT"));
+            }
+
+            CVPixelBufferRef PixelBuffer;
+            CVReturn Result = CVPixelBufferPoolCreatePixelBuffer(nullptr, PixelBufferPool, &PixelBuffer);
+            if(Result != kCVReturnSuccess)
+            {
+                return FAVResult(EAVResult::Error, TEXT("Failed to create pixel buffer"), TEXT("VT"));
+            }
+
+            Result = CVPixelBufferLockBaseAddress(PixelBuffer, 0);
+            if (Result != kCVReturnSuccess) 
+            {
+                return FAVResult(EAVResult::Error, TEXT("Failed to lock base address"), TEXT("VT"));
+            }
+
+            // NOTE (belchy06): GetBytes assumes the raw texture has been created with with TexCreate_CPUReadback
+            mtlpp::Texture* RawTexture = Resource->GetRaw();
+            uint32_t Width = RawTexture->GetWidth();
+            uint32_t Height = RawTexture->GetHeight();
+            TArray<uint8> PixelBytes;
+            PixelBytes.SetNumZeroed(4 * Width * Height);
+
+            Resource->GetRaw()->GetBytes(PixelBytes.GetData(), 4 * Width, mtlpp::Region(0, 0, Width, Height), 0);
+            FMemory::BigBlockMemcpy(reinterpret_cast<uint8*>(CVPixelBufferGetBaseAddressOfPlane(PixelBuffer, 0)), PixelBytes.GetData(), PixelBytes.Num());
+            CVPixelBufferUnlockBaseAddress(PixelBuffer, 0);
+            
+            CMTime PresentationTime = CMTimeMake(Timestamp, 1000);
+            CFDictionaryRef FrameProperties = nullptr;
+            if (bForceKeyframe) 
+            {
+                CFTypeRef Keys[] = {kVTEncodeFrameOptionKey_ForceKeyFrame};
+                CFTypeRef Values[] = {kCFBooleanTrue};
+                FrameProperties = CFDictionaryCreate(kCFAllocatorDefault, Keys, Values, 1, &kCFTypeDictionaryKeyCallBacks, &kCFTypeDictionaryValueCallBacks);
+            }
+
+            TUniquePtr<EncodeParams> Params = MakeUnique<EncodeParams>();
+            Params.Reset(new EncodeParams(this->GetAppliedConfig().Codec, PresentationTime));
+
+            OSStatus Status = VTCompressionSessionEncodeFrame(Encoder, PixelBuffer, PresentationTime, kCMTimeInvalid, FrameProperties, (void*)Params.Release(), nullptr);
+
+            CONDITIONAL_RELEASE(FrameProperties);
+
+            if (PixelBuffer) 
+            {
+                CVPixelBufferRelease(PixelBuffer);
+                PixelBuffer = nullptr;
+            }
+        
+
+            if(Status != 0)
+            {
+                return FAVResult(EAVResult::Error, TEXT("Failed to encode"), TEXT("VT"));
+            }
+            
+            return EAVResult::Success;
+        }
+        else
+        {
+            // Flush encoder
+            OSStatus Status = VTCompressionSessionCompleteFrames(Encoder, kCMTimeInvalid);
+            if(Status != 0)
+            {
+                return FAVResult(EAVResult::Error, TEXT("Failed to flush"), TEXT("VT"));
+            }
+
+            return EAVResult::Success;
+        }
+    }
+
+	return FAVResult(EAVResult::ErrorInvalidState, TEXT("Encoder not open"), TEXT("VT"));
+}
+
+
+template <typename TResource>
+FAVResult TVideoEncoderVT<TResource>::ReceivePacket(FVideoPacket& OutPacket)
+{
+	if (IsOpen())
+	{
+		if (Packets.Dequeue(OutPacket))
+		{
+			return EAVResult::Success;
+		}
+
+		return EAVResult::PendingInput;
+	}
+
+	return FAVResult(EAVResult::ErrorInvalidState, TEXT("Encoder not open"), TEXT("VT"));
+}
+
+template <typename TResource>
+FAVResult TVideoEncoderVT<TResource>::HandlePacket(void* Params, OSStatus Status, VTEncodeInfoFlags InfoFlags, CMSampleBufferRef& SampleBuffer)
+{
+	if (IsOpen())
+	{
+        if(Status != 0)
+        {
+            return FAVResult(EAVResult::Error, TEXT("Failed to encode"), TEXT("VT"));
+        }
+
+        if(InfoFlags & kVTEncodeInfo_FrameDropped)
+        {
+            return FAVResult(EAVResult::Error, TEXT("Frame dropped"), TEXT("VT"));
+        }
+
+        bool bIsKeyframe = false;
+        CFArrayRef Attachments = CMSampleBufferGetSampleAttachmentsArray(SampleBuffer, 0);
+        if (Attachments != nullptr && CFArrayGetCount(Attachments)) 
+        {
+            CFDictionaryRef Attachment = static_cast<CFDictionaryRef>(CFArrayGetValueAtIndex(Attachments, 0));
+            bIsKeyframe = !CFDictionaryContainsKey(Attachment, kCMSampleAttachmentKey_NotSync);
+        }
+        
+        // SampleBuffer is an MPEG4 container stream, but we need to convert it to its raw bitstream
+        EncodeParams* Config = static_cast<EncodeParams*>(Params);
+
+        int QP = 0;
+        TArray<uint8> Bitstream;
+        if(Config->Codec == kCMVideoCodecType_H264) 
+        {
+            if(!NaluRewriter::H264CMSampleBufferToAnnexBBuffer(SampleBuffer, bIsKeyframe, Bitstream))
+            {
+                return FAVResult(EAVResult::Error, TEXT("Failed to extract bitstream"), TEXT("VT"));
+            }
+        }
+        else if(Config->Codec == kCMVideoCodecType_HEVC) 
+        {
+            if(!NaluRewriter::H265CMSampleBufferToAnnexBBuffer(SampleBuffer, bIsKeyframe, Bitstream))
+            {
+                return FAVResult(EAVResult::Error, TEXT("Failed to extract bitstream"), TEXT("VT"));
+            }
+        }
+        else
+        {
+            return FAVResult(EAVResult::ErrorUnsupported, TEXT("Unsupported codec"), TEXT("VT"));
+        }
+
+        TSharedPtr<uint8> const CopiedData = MakeShareable(new uint8[Bitstream.Num()]);
+		FMemory::BigBlockMemcpy(CopiedData.Get(), Bitstream.GetData(), Bitstream.Num());
+
+        if(Config->Codec == kCMVideoCodecType_H264)
+        {
+            using namespace UE::AVCodecCore::H264;
+			TArray<Slice_t> Slices;
+            FAVResult Result = H264.Parse(FVideoPacket(CopiedData, Bitstream.Num(), 0, 0, 0, false), Slices);
+            if(Result != EAVResult::Success)
+            {
+                return FAVResult(EAVResult::Error, TEXT("Failed to parse bitstream"), TEXT("VT"));
+            }
+
+            QP = H264.GetLastSliceQP(Slices).Get(0);
+        }
+        else if(Config->Codec == kCMVideoCodecType_HEVC)
+        {
+            using namespace UE::AVCodecCore::H265;
+			TArray<Slice_t> Slices;
+            FAVResult Result = H265.Parse(FVideoPacket(CopiedData, Bitstream.Num(), 0, 0, 0, false), Slices);
+            if(Result != EAVResult::Success)
+            {
+                return FAVResult(EAVResult::Error, TEXT("Failed to parse bitstream"), TEXT("VT"));
+            }
+
+            QP = H265.GetLastSliceQP(Slices).Get(0);
+            FAVResult::Log(EAVResult::Warning, TEXT("H265 QP"), TEXT("VT"), QP);
+        }
+
+		return Packets.Enqueue(FVideoPacket(CopiedData, Bitstream.Num(), CMClockConvertHostTimeToSystemUnits(Config->Timestamp), ++FrameCount, (uint32_t)QP, bIsKeyframe)) ? EAVResult::Success : EAVResult::Error;
+	}
+
+	return FAVResult(EAVResult::ErrorInvalidState, TEXT("Encoder not open"), TEXT("VT"));
+}
