Fix: Missed change in CL29731236

diff --git a/Engine/Plugins/Experimental/AVCodecs/NVCodecs/Source/NVDEC/Public/Video/Decoders/Configs/VideoDecoderConfigNVDEC.h b/Engine/Plugins/Experimental/AVCodecs/NVCodecs/Source/NVDEC/Public/Video/Decoders/Configs/VideoDecoderConfigNVDEC.h
--- a/Engine/Plugins/Experimental/AVCodecs/NVCodecs/Source/NVDEC/Public/Video/Decoders/Configs/VideoDecoderConfigNVDEC.h
+++ b/Engine/Plugins/Experimental/AVCodecs/NVCodecs/Source/NVDEC/Public/Video/Decoders/Configs/VideoDecoderConfigNVDEC.h
@@ -57,7 +57,7 @@ public:
 		return !(*this == Other);
 	}
 
-	FAVResult Parse(TSharedRef<FAVInstance> const& Instance, FVideoPacket const& Packet, TArray<FParsedPicture>& OutPictures);
+	FAVResult Parse(TSharedRef<FAVInstance> const& Instance, FVideoPacket const& Packet, TArray<FParsedPicture>& OutPictures) { return EAVResult::Error; }
 
 private:
 	int RefPicIdx[16];
