Fix: Re-enable NVDEC

diff --git a/Engine/Plugins/Experimental/AVCodecs/NVCodecs/Source/NVDEC/Private/NVDEC.cpp b/Engine/Plugins/Experimental/AVCodecs/NVCodecs/Source/NVDEC/Private/NVDEC.cpp
--- a/Engine/Plugins/Experimental/AVCodecs/NVCodecs/Source/NVDEC/Private/NVDEC.cpp
+++ b/Engine/Plugins/Experimental/AVCodecs/NVCodecs/Source/NVDEC/Private/NVDEC.cpp
@@ -53,6 +53,7 @@ FNVDEC::FNVDEC()
 			LOAD_FROM_DLL(cuvidCreateDecoder);
 			LOAD_FROM_DLL(cuvidParseVideoData);
 			LOAD_FROM_DLL(cuvidCreateVideoParser);
+			LOAD_FROM_DLL(cuvidDestroyVideoParser);
 			LOAD_FROM_DLL(cuvidCtxLockCreate);
 			LOAD_FROM_DLL(cuvidCtxLockDestroy);
 			LOAD_FROM_DLL(cuvidDecodePicture);
diff --git a/Engine/Plugins/Experimental/AVCodecs/NVCodecs/Source/NVDEC/Private/NVDECModule.cpp b/Engine/Plugins/Experimental/AVCodecs/NVCodecs/Source/NVDEC/Private/NVDECModule.cpp
--- a/Engine/Plugins/Experimental/AVCodecs/NVCodecs/Source/NVDEC/Private/NVDECModule.cpp
+++ b/Engine/Plugins/Experimental/AVCodecs/NVCodecs/Source/NVDEC/Private/NVDECModule.cpp
@@ -14,7 +14,6 @@ class FNVDECModule : public IModuleInterface
 public:
 	virtual void StartupModule() override
 	{
-		// TODO (william.belcher): NVDEC decoding doesn't currently work
 		FVideoDecoder
 			::RegisterPermutationsOf<FVideoDecoderNVDEC>
 			::With<
diff --git a/Engine/Plugins/Experimental/AVCodecs/NVCodecs/Source/NVDEC/Private/Video/Decoders/VideoDecoderNVDEC.cpp b/Engine/Plugins/Experimental/AVCodecs/NVCodecs/Source/NVDEC/Private/Video/Decoders/VideoDecoderNVDEC.cpp
--- a/Engine/Plugins/Experimental/AVCodecs/NVCodecs/Source/NVDEC/Private/Video/Decoders/VideoDecoderNVDEC.cpp
+++ b/Engine/Plugins/Experimental/AVCodecs/NVCodecs/Source/NVDEC/Private/Video/Decoders/VideoDecoderNVDEC.cpp
@@ -7,6 +7,24 @@
 #include "NV12_to_BGRA8.cuh"
 #include "P010_to_ABGR10.cuh"
 
+namespace Internal 
+{
+    int HandleVideoSequenceCallback(void* UserData, CUVIDEOFORMAT* VideoFormat)
+	{
+		return static_cast<FVideoDecoderNVDEC*>(UserData)->HandleVideoSequence(VideoFormat);
+	}
+
+	int HandlePictureDCodecallback(void *UserData, CUVIDPICPARAMS *PicParams)
+	{
+		return static_cast<FVideoDecoderNVDEC*>(UserData)->HandlePictureDecode(PicParams);
+	}
+
+	int HandlePictureDisplayCallback(void *UserData, CUVIDPARSERDISPINFO *DispInfo)
+	{
+		return static_cast<FVideoDecoderNVDEC*>(UserData)->HandlePictureDisplay(DispInfo);
+	}
+}
+
 FVideoDecoderNVDEC::~FVideoDecoderNVDEC()
 {
 	Close();
@@ -25,6 +43,27 @@ FAVResult FVideoDecoderNVDEC::Open(TSharedRef<FAVDevice> const& NewDevice, TShar
 
 	bIsOpen = true;
 
+	CUVIDPARSERPARAMS VideoParserParameters = {};
+	FVideoDecoderConfigNVDEC const& PendingConfig = GetPendingConfig();
+    VideoParserParameters.CodecType = PendingConfig.CodecType;
+    VideoParserParameters.ulMaxNumDecodeSurfaces = 1;
+	// TODO (william.belcher): bLowLatency ? 0 : 1;
+    VideoParserParameters.ulMaxDisplayDelay = 0;
+    VideoParserParameters.pUserData = this;
+    VideoParserParameters.pfnSequenceCallback = Internal::HandleVideoSequenceCallback;
+    VideoParserParameters.pfnDecodePicture = Internal::HandlePictureDCodecallback;
+    VideoParserParameters.pfnDisplayPicture = Internal::HandlePictureDisplayCallback;
+
+	if (FAPI::Get<FNVDEC>().cuvidCreateVideoParser(&Parser, &VideoParserParameters) != CUDA_SUCCESS)
+	{
+		return FAVResult(EAVResult::ErrorCreating, TEXT("Failed to create parser"), TEXT("NVDEC"));
+	}
+
+	if (FAPI::Get<FNVDEC>().cuvidCtxLockCreate(&CtxLock, GetDevice()->GetContext<FVideoContextCUDA>()->Raw) != CUDA_SUCCESS)
+	{
+		return FAVResult(EAVResult::ErrorCreating, TEXT("Failed to context lock decoder"), TEXT("NVDEC"));
+	}
+
 	return EAVResult::Success;
 }
 
@@ -37,12 +76,36 @@ void FVideoDecoderNVDEC::Close()
 		CUresult const Result = FAPI::Get<FNVDEC>().cuvidDestroyDecoder(Decoder);
 		if (Result != CUDA_SUCCESS)
 		{
-			FAVResult::Log(EAVResult::ErrorDestroying, TEXT("Failed to destroy NVENC encoder"), TEXT("NVDEC"), Result);
+			FAVResult::Log(EAVResult::ErrorDestroying, TEXT("Failed to destroy NVDEC decoder"), TEXT("NVDEC"), Result);
 		}
 
 		Decoder = nullptr;
 	}
 
+	if (Parser != nullptr)
+	{
+		FCUDAContextScope const ContextGuard(GetDevice()->GetContext<FVideoContextCUDA>()->Raw);
+
+		CUresult const Result = FAPI::Get<FNVDEC>().cuvidDestroyVideoParser(Parser);
+		if (Result != CUDA_SUCCESS)
+		{
+			FAVResult::Log(EAVResult::ErrorDestroying, TEXT("Failed to destroy NVDEC parser"), TEXT("NVDEC"), Result);
+		}
+	}
+
+	// We should only run this block if we've ever opened the decoder before
+	// Basically just prevents this block from being run by the call to Close() in Open()
+	if (IsOpen())
+	{
+		FCUDAContextScope const ContextGuard(GetDevice()->GetContext<FVideoContextCUDA>()->Raw);
+
+		CUresult const Result = FAPI::Get<FNVDEC>().cuvidCtxLockDestroy(CtxLock);
+		if (Result != CUDA_SUCCESS)
+		{
+			FAVResult::Log(EAVResult::ErrorDestroying, TEXT("Failed to destroy NVDEC context lock"), TEXT("NVDEC"), Result);
+		}
+	}
+
 	bIsOpen = false;
 }
 
@@ -60,43 +123,58 @@ FAVResult FVideoDecoderNVDEC::ApplyConfig()
 		{
 			if (IsInitialized())
 			{
-				/*
-				TODO(Andrew) Can be reconfigured? See https://docs.nvidia.com/video-technologies/video-codec-sdk/nvenc-video-encoder-api-prog-guide/#reconfigure-api
-				if (AppliedConfig.maxEncodeWidth == PendingConfig.maxEncodeWidth
-					&& AppliedConfig.maxEncodeHeight == PendingConfig.maxEncodeHeight
-					&& AppliedConfig.enablePTD == PendingConfig.enablePTD
-					&& AppliedConfig.enableEncodeAsync == PendingConfig.enableEncodeAsync
-					&& AppliedConfig.encodeConfig->gopLength == PendingConfig.encodeConfig->gopLength
-					&& AppliedConfig.encodeConfig->frameIntervalP == PendingConfig.encodeConfig->frameIntervalP
-					&& AppliedConfig.encodeConfig->encodeCodecConfig.h264Config.idrPeriod == PendingConfig.encodeConfig->encodeCodecConfig.h264Config.idrPeriod)
+				// If any of these change, we need to re-initialization the decoder
+				if (AppliedConfig.OutputFormat == PendingConfig.OutputFormat 
+					&& AppliedConfig.ChromaFormat == PendingConfig.ChromaFormat)
 				{
-					NV_ENC_STRUCT(NV_ENC_RECONFIGURE_PARAMS, ReconfigureParams);
-					FMemory::Memcpy(&ReconfigureParams.reInitEncodeParams, &PendingConfig, sizeof(NV_ENC_INITIALIZE_PARAMS));
-					ReconfigureParams.forceIDR = AppliedConfig.encodeWidth != PendingConfig.encodeWidth || AppliedConfig.encodeHeight != PendingConfig.encodeHeight;
+					bool bResChange = AppliedConfig.ulWidth != PendingConfig.ulWidth || AppliedConfig.ulHeight != PendingConfig.ulHeight;
 
-					NVENCSTATUS const Result = FAPI::Get<FNVENC_Context>().nvEncReconfigureEncoder(Encoder, &ReconfigureParams);
-					if (Result != NV_ENC_SUCCESS)
+					if(bResChange)
 					{
-						return FAVResult(EAVResult::Error, FString::Printf(TEXT("Failed to update NVENC encoder configuration: %s"), *FAPI::Get<FNVENC_Context>().GetErrorString(Encoder, Result)));
+						CUVIDRECONFIGUREDECODERINFO ReconfigParams = { 0 };
+
+						ReconfigParams.ulWidth = PendingConfig.ulWidth;
+   						ReconfigParams.ulHeight = PendingConfig.ulHeight;
+
+    					ReconfigParams.display_area.bottom = PendingConfig.display_area.bottom;
+    					ReconfigParams.display_area.top = PendingConfig.display_area.top;
+    					ReconfigParams.display_area.left = PendingConfig.display_area.left;
+    					ReconfigParams.display_area.right = PendingConfig.display_area.right;
+    					ReconfigParams.ulTargetWidth = PendingConfig.ulWidth;
+    					ReconfigParams.ulTargetHeight = PendingConfig.ulHeight;
+    
+    					ReconfigParams.ulNumDecodeSurfaces = GetNumDecodeSurfaces(PendingConfig.CodecType, PendingConfig.ulWidth, PendingConfig.ulHeight);;
+
+						FCUDAContextScope const ContextGuard(GetDevice()->GetContext<FVideoContextCUDA>()->Raw);
+						CUresult const Result = FAPI::Get<FNVDEC>().cuvidReconfigureDecoder(Decoder, &ReconfigParams);
+						if (Result != CUDA_SUCCESS)
+						{
+							Close();
+							return FAVResult(EAVResult::Error, TEXT("Failed to reconfigure decoder"), TEXT("NVDEC"), Result);
+						}
 					}
 				}
 				else
 				{
-					// TODO: Destroy and recreate with original session
-					unimplemented();
+					FCUDAContextScope const ContextGuard(GetDevice()->GetContext<FVideoContextCUDA>()->Raw);
+					CUresult const Result = FAPI::Get<FNVDEC>().cuvidDestroyDecoder(Decoder);
+					if (Result != CUDA_SUCCESS)
+					{
+						Close();
+						return FAVResult(EAVResult::ErrorDestroying, TEXT("Failed to destroy NVDEC decoder"), TEXT("NVDEC"), Result);
+					}
+
+					Decoder = nullptr;
 				}
-				*/
 			}
 
 			if (!IsInitialized())
 			{
 				FCUDAContextScope const ContextGuard(GetDevice()->GetContext<FVideoContextCUDA>()->Raw);
-
 				CUresult const Result = FAPI::Get<FNVDEC>().cuvidCreateDecoder(&Decoder, &const_cast<FVideoDecoderConfigNVDEC&>(PendingConfig));
 				if (Result != CUDA_SUCCESS)
 				{
 					Close();
-
 					return FAVResult(EAVResult::ErrorCreating, TEXT("Failed to create decoder"), TEXT("NVDEC"), Result);
 				}
 			}
@@ -135,50 +213,23 @@ FAVResult FVideoDecoderNVDEC::SendPacket(FVideoPacket const& Packet)
 		return FAVResult(EAVResult::ErrorInvalidState, TEXT("Decoder not open"), TEXT("NVDEC"));
 	}
 
-	FVideoDecoderConfigNVDEC& PendingConfig = EditPendingConfig();
-
-	TArray<FVideoDecoderConfigNVDEC::FParsedPicture> Pictures;
-
-	FAVResult Result = PendingConfig.Parse(GetInstance().ToSharedRef(), Packet, Pictures);
-	if (Result.IsNotSuccess())
+	CUVIDSOURCEDATAPACKET CUPacket = {0};
+    CUPacket.payload = Packet.DataPtr.Get();
+    CUPacket.payload_size = Packet.DataSize;
+    CUPacket.flags = CUVID_PKT_TIMESTAMP;
+	// TODO (william.belcher): Adding this flag assumes that each Packet contains a full frame, but is required
+	// to prevent the parser from adding a 1 frame delay
+	CUPacket.flags |= CUVID_PKT_ENDOFPICTURE;
+    CUPacket.timestamp = 0;
+    if (!Packet.DataPtr.Get() || Packet.DataSize == 0) 
 	{
-		return Result;
-	}
-
-	if (Pictures.Num() == 0)
-	{
-		return EAVResult::PendingInput;
-	}
+        CUPacket.flags |= CUVID_PKT_ENDOFSTREAM;
+    }
 
-	for (auto& Picture : Pictures)
+    CUresult const CUResult = FAPI::Get<FNVDEC>().cuvidParseVideoData(Parser, &CUPacket);
+	if (CUResult != CUDA_SUCCESS)
 	{
-		if (FramesCount + 1 >= static_cast<int64>(Picture.DecodeCreateInfo.ulNumDecodeSurfaces))
-		{
-			return FAVResult(EAVResult::PendingOutput, TEXT("Decode buffer full"), TEXT("NVDEC"));
-		}
-
-		CUVIDPROCPARAMS MapParams = {};
-
-		FMemory::Memcpy(&static_cast<CUVIDDECODECREATEINFO&>(PendingConfig), &Picture.DecodeCreateInfo, sizeof(CUVIDDECODECREATEINFO));
-
-		Result = ApplyConfig();
-		if (Result.IsNotSuccess())
-		{
-			return Result;
-		}
-
-		CUresult const CUResult = FAPI::Get<FNVDEC>().cuvidDecodePicture(Decoder, &Picture);
-		if (CUResult != CUDA_SUCCESS)
-		{
-			return FAVResult(EAVResult::Error, TEXT("Failed to decode frame"), TEXT("NVDEC"), CUResult);
-		}
-
-		MapParams.progressive_frame = PendingConfig.DeinterlaceMode;
-		MapParams.second_field = Picture.second_field;
-		MapParams.top_field_first = Picture.bottom_field_flag != 0;
-
-		Frames.Enqueue({ Picture.CurrPicIdx, AppliedConfig.ulTargetWidth, AppliedConfig.ulTargetHeight, GetAppliedConfig().OutputFormat, MapParams });
-		++FramesCount;
+		return FAVResult(EAVResult::Error, TEXT("Failed to parse video data"), TEXT("NVDEC"), CUResult);
 	}
 
 	return EAVResult::Success;
@@ -194,7 +245,7 @@ FAVResult FVideoDecoderNVDEC::ReceiveFrame(TResolvableVideoResource<FVideoResour
 			FCUDAContextScope const ContextGuard(GetDevice()->GetContext<FVideoContextCUDA>()->Raw);
 
 			CUresult Result = CUDA_SUCCESS;
-			CUVIDGETDECODESTATUS DecodeStatus = {};
+			CUVIDGETDECODESTATUS DecodeStatus = { };
 
 			do
 			{
@@ -218,7 +269,7 @@ FAVResult FVideoDecoderNVDEC::ReceiveFrame(TResolvableVideoResource<FVideoResour
 			{
 				FAVResult::Log(EAVResult::Warning, TEXT("Error when decoding frame"), TEXT("NVDEC"), DecodeStatus.decodeStatus);
 			}
-			
+
 			// Map the decoder output buffer
 			CUdeviceptr MapSurface = 0;
 			uint32 MapPitch = 0;
@@ -230,62 +281,6 @@ FAVResult FVideoDecoderNVDEC::ReceiveFrame(TResolvableVideoResource<FVideoResour
 			}
 
 			FVideoDescriptor ResourceDescriptor = {};
-
-#if 0 // Old path with automatic conversion
-			TFunction<CUresult(CUdeviceptr, CUarray, uint32, uint32, uint32)> ConversionFunction;
-
-			// TODO (aidan) handle conversion to non-RGB textures
-			switch (Frame.SurfaceFormat)
-			{
-			case cudaVideoSurfaceFormat_NV12:
-				ResourceDescriptor = FVideoDescriptor(EVideoFormat::BGRA, Frame.Width, Frame.Height);
-				ConversionFunction = NV12_to_BGRA8;
-				break;
-			case cudaVideoSurfaceFormat_P016:
-				ResourceDescriptor = FVideoDescriptor(EVideoFormat::ABGR10, Frame.Width, Frame.Height);
-				ConversionFunction = P010_to_ABGR10; // TODO
-				break;
-			/*case cudaVideoSurfaceFormat_YUV444:
-				ResourceDescriptor = FVideoDescriptor(EVideoFormat::YUV444, Frame.Width, Frame.Height);
-				break;
-			case cudaVideoSurfaceFormat_YUV444_16Bit:
-				ResourceDescriptor = FVideoDescriptor(EVideoFormat::YUV444_16, Frame.Width, Frame.Height);*/
-			default:
-				return FAVResult(EAVResult::ErrorUnsupported, FString::Printf(TEXT("Currently do not support the recieved OutputFormat (%d)"), Frame.SurfaceFormat));
-			}
-
-			if (!InOutResource.Resolve(GetDevice(), ResourceDescriptor))
-			{
-				return FAVResult(EAVResult::ErrorResolving, TEXT("Failed to resolve frame resource"), TEXT("NVDEC"));
-			}
-
-			InOutResource->Lock();
-
-			Result = ConversionFunction(MapSurface, InOutResource->GetRaw(), InOutResource->GetWidth(), InOutResource->GetHeight(), MapPitch);
-
-			InOutResource->Unlock();
-
-			if (Result != CUDA_SUCCESS)
-			{
-				return FAVResult(EAVResult::Error, TEXT("Failed to run CUDA kernal on decoded frame"), TEXT("NVDEC"), Result);
-			}
-
-			// START DEBUG
-
-#if DEBUG_DUMP_TO_DISK
-
-			TArray64<uint8> OutData;
-
-			InOutResource->ReadData(OutData);
-
-			FString SaveName = FString::Printf(TEXT("%s/DumpOutput/image%05d.%s"), *FPaths::ProjectSavedDir(), Frame.SurfaceIndex, Frame.SurfaceFormat == cudaVideoSurfaceFormat_NV12 ? TEXT("nv12") : TEXT("p016"));
-
-			FFileHelper::SaveArrayToFile(OutData, *SaveName);
-
-#endif
-
-			// END DEBUG
-#else // new code path just returns in the format that the decoder outputs and is up to the user to transform it
 			switch (Frame.SurfaceFormat)
 			{
 				case cudaVideoSurfaceFormat_NV12:
@@ -318,10 +313,7 @@ FAVResult FVideoDecoderNVDEC::ReceiveFrame(TResolvableVideoResource<FVideoResour
 			if (Result != CUDA_SUCCESS)
 			{
 				return FAVResult(EAVResult::ErrorUnmapping, TEXT("Failed to unmap decoded frame"), TEXT("NVDEC"), Result);
-			}
-			
-
-#endif // end else !0		
+			}	
 
 			return EAVResult::Success;
 		}
@@ -330,4 +322,133 @@ FAVResult FVideoDecoderNVDEC::ReceiveFrame(TResolvableVideoResource<FVideoResour
 	}
 
 	return FAVResult(EAVResult::ErrorInvalidState, TEXT("Decoder not open"), TEXT("NVDEC"));
+}
+
+unsigned long FVideoDecoderNVDEC::GetNumDecodeSurfaces(cudaVideoCodec Codec, unsigned int Width, unsigned int Height)
+{
+	if (Codec == cudaVideoCodec_VP9) 
+	{
+        return 12;
+    }
+
+    if (Codec == cudaVideoCodec_H264 || Codec == cudaVideoCodec_H264_SVC || Codec == cudaVideoCodec_H264_MVC) 
+	{
+        // assume worst-case of 20 decode surfaces for H264
+        return 20;
+    }
+
+    if (Codec == cudaVideoCodec_HEVC) 
+	{
+        // ref HEVC spec: A.4.1 General tier and level limits
+        // currently assuming level 6.2, 8Kx4K
+        int MaxLumaPS = 35651584;
+        int MaxDpbPicBuf = 6;
+        int PicSizeInSamplesY = (int)(Width * Height);
+        int MaxDpbSize;
+
+        if (PicSizeInSamplesY <= (MaxLumaPS>>2))
+		{
+			MaxDpbSize = MaxDpbPicBuf * 4;
+		}
+        else if (PicSizeInSamplesY <= (MaxLumaPS>>1))
+		{
+			MaxDpbSize = MaxDpbPicBuf * 2;
+		}
+        else if (PicSizeInSamplesY <= ((3*MaxLumaPS)>>2))
+		{
+			MaxDpbSize = (MaxDpbPicBuf * 4) / 3;
+		}
+        else
+		{
+			MaxDpbSize = MaxDpbPicBuf;
+		}
+
+        return FMath::Min(MaxDpbSize, 16) + 4;
+    }
+
+    return 8;
+}
+
+int FVideoDecoderNVDEC::HandleVideoSequence(CUVIDEOFORMAT *VideoFormat)
+{
+	int NumDecodeSurfaces = GetNumDecodeSurfaces(VideoFormat->codec, VideoFormat->coded_width, VideoFormat->coded_height);
+
+	CUVIDDECODECAPS DecodeCaps;
+	memset(&DecodeCaps, 0, sizeof(DecodeCaps));
+
+	DecodeCaps.eCodecType = VideoFormat->codec;
+	DecodeCaps.eChromaFormat = VideoFormat->chroma_format;
+	DecodeCaps.nBitDepthMinus8 = VideoFormat->bit_depth_luma_minus8;
+
+	if (!GetCapability(DecodeCaps))
+	{
+		FAVResult::Log(EAVResult::ErrorUnsupported, TEXT("Codec not supported on this GPU"), TEXT("NVDEC"));
+		return NumDecodeSurfaces;
+	}
+
+	if ((VideoFormat->coded_width > DecodeCaps.nMaxWidth) || (VideoFormat->coded_height > DecodeCaps.nMaxHeight))
+	{
+		FAVResult::Log(EAVResult::ErrorUnsupported, FString::Printf(TEXT("Resolution: %dx%d\nMax Supported Resolution: %dx%d\n Resolution not supported on this GPU"), VideoFormat->coded_width, VideoFormat->coded_height, DecodeCaps.nMaxWidth, DecodeCaps.nMaxHeight), TEXT("NVDEC"));
+		return NumDecodeSurfaces;
+	}
+
+	if ((VideoFormat->coded_width >> 4) * (VideoFormat->coded_height >> 4) > DecodeCaps.nMaxMBCount)
+	{
+		FAVResult::Log(EAVResult::ErrorUnsupported, FString::Printf(TEXT("MBCount: %d\nMax Supported mbcnt: %d\n MBCount not supported on this GPU"), (VideoFormat->coded_width >> 4) * (VideoFormat->coded_height >> 4), DecodeCaps.nMaxMBCount), TEXT("NVDEC"));
+		return NumDecodeSurfaces;
+	}
+
+	FVideoDecoderConfigNVDEC& PendingConfig = EditPendingConfig();
+	PendingConfig.CodecType = VideoFormat->codec;
+	PendingConfig.ChromaFormat = VideoFormat->chroma_format;
+    PendingConfig.OutputFormat = VideoFormat->bit_depth_luma_minus8 ? cudaVideoSurfaceFormat_P016 : cudaVideoSurfaceFormat_NV12;
+    PendingConfig.bitDepthMinus8 = VideoFormat->bit_depth_luma_minus8;
+    PendingConfig.DeinterlaceMode = cudaVideoDeinterlaceMode_Weave;
+    PendingConfig.ulNumOutputSurfaces = 2;
+    // With PreferCUVID, JPEG is still decoded by CUDA while video is decoded by NVDEC hardware
+    PendingConfig.ulCreationFlags = cudaVideoCreate_PreferCUVID;
+    PendingConfig.ulNumDecodeSurfaces = NumDecodeSurfaces;
+    PendingConfig.vidLock = CtxLock;
+    PendingConfig.ulWidth = VideoFormat->coded_width;
+    PendingConfig.ulHeight = VideoFormat->coded_height;
+    PendingConfig.ulMaxWidth = VideoFormat->coded_width;
+    PendingConfig.ulMaxHeight = VideoFormat->coded_height;
+
+	// TODO (william.belcher): Add support for cropping and resizing
+	PendingConfig.ulTargetWidth = VideoFormat->coded_width;
+    PendingConfig.ulTargetHeight = VideoFormat->coded_height;
+
+	FAVResult Result = ApplyConfig();
+	if (Result.IsNotSuccess())
+	{
+		FAVResult::Log(EAVResult::Error, TEXT("Failed to apply decoder config"), TEXT("NVDEC"), Result);
+	}
+
+	return NumDecodeSurfaces;
+}
+
+int FVideoDecoderNVDEC::HandlePictureDecode(CUVIDPICPARAMS *PicParams)
+{
+	CUresult const Result = FAPI::Get<FNVDEC>().cuvidDecodePicture(Decoder, PicParams);
+	if (Result != CUDA_SUCCESS)
+	{
+		FAVResult::Log(EAVResult::Error, TEXT("Failed to decode frame"), TEXT("NVDEC"), Result);
+		return 0;
+	}
+	return 1;
+}
+
+int FVideoDecoderNVDEC::HandlePictureDisplay(CUVIDPARSERDISPINFO *DispInfo)
+{
+	CUVIDPROCPARAMS VideoProcessingParameters = {};
+    VideoProcessingParameters.progressive_frame = DispInfo->progressive_frame;
+    VideoProcessingParameters.second_field = DispInfo->repeat_first_field + 1;
+    VideoProcessingParameters.top_field_first = DispInfo->top_field_first;
+    VideoProcessingParameters.unpaired_field = DispInfo->repeat_first_field < 0;
+    VideoProcessingParameters.output_stream = 0;
+
+	Frames.Enqueue({ DispInfo->picture_index, AppliedConfig.ulTargetWidth, AppliedConfig.ulTargetHeight, GetAppliedConfig().OutputFormat, VideoProcessingParameters });
+	++FramesCount;
+
+	return 0;
 }
\ No newline at end of file
diff --git a/Engine/Plugins/Experimental/AVCodecs/NVCodecs/Source/NVDEC/Public/NVDEC.h b/Engine/Plugins/Experimental/AVCodecs/NVCodecs/Source/NVDEC/Public/NVDEC.h
--- a/Engine/Plugins/Experimental/AVCodecs/NVCodecs/Source/NVDEC/Public/NVDEC.h
+++ b/Engine/Plugins/Experimental/AVCodecs/NVCodecs/Source/NVDEC/Public/NVDEC.h
@@ -29,6 +29,7 @@ struct NV_DECODE_API_FUNCTION_LIST
 	typedef CUresult (*cuvidCreateDecoderPtr)(CUvideodecoder*, CUVIDDECODECREATEINFO*);
 	typedef CUresult (*cuvidParseVideoDataPtr)(CUvideoparser, CUVIDSOURCEDATAPACKET*);
 	typedef CUresult (*cuvidCreateVideoParserPtr)(CUvideoparser*, CUVIDPARSERPARAMS*);
+	typedef CUresult (*cuvidDestroyVideoParserPtr)(CUvideoparser);
 	typedef CUresult (*cuvidCtxLockCreatePtr)(CUvideoctxlock*, CUcontext);
 	typedef CUresult (*cuvidCtxLockDestroyPtr)(CUvideoctxlock);
 	typedef CUresult (*cuvidDecodePicturePtr)(CUvideodecoder, CUVIDPICPARAMS*);
@@ -42,6 +43,7 @@ struct NV_DECODE_API_FUNCTION_LIST
 	cuvidCreateDecoderPtr cuvidCreateDecoder = nullptr;
 	cuvidParseVideoDataPtr cuvidParseVideoData = nullptr;
 	cuvidCreateVideoParserPtr cuvidCreateVideoParser = nullptr;
+	cuvidDestroyVideoParserPtr cuvidDestroyVideoParser = nullptr;
 	cuvidCtxLockCreatePtr cuvidCtxLockCreate = nullptr;
 	cuvidCtxLockDestroyPtr cuvidCtxLockDestroy = nullptr;
 	cuvidDecodePicturePtr cuvidDecodePicture = nullptr;
diff --git a/Engine/Plugins/Experimental/AVCodecs/NVCodecs/Source/NVDEC/Public/Video/Decoders/VideoDecoderNVDEC.h b/Engine/Plugins/Experimental/AVCodecs/NVCodecs/Source/NVDEC/Public/Video/Decoders/VideoDecoderNVDEC.h
--- a/Engine/Plugins/Experimental/AVCodecs/NVCodecs/Source/NVDEC/Public/Video/Decoders/VideoDecoderNVDEC.h
+++ b/Engine/Plugins/Experimental/AVCodecs/NVCodecs/Source/NVDEC/Public/Video/Decoders/VideoDecoderNVDEC.h
@@ -16,6 +16,8 @@ class NVDEC_API FVideoDecoderNVDEC : public TVideoDecoder<FVideoResourceCUDA, FV
 private:
 	uint8 bIsOpen : 1;
 	CUvideodecoder Decoder = nullptr;
+	CUvideoparser Parser = nullptr;
+	CUvideoctxlock CtxLock;
 
 	struct FFrame
 	{
@@ -47,4 +49,19 @@ public:
 	virtual FAVResult SendPacket(FVideoPacket const& Packet) override;
 	
 	virtual FAVResult ReceiveFrame(TResolvableVideoResource<FVideoResourceCUDA>& InOutResource) override;
+
+public:
+	int HandleVideoSequence(CUVIDEOFORMAT *VideoFormat);
+    int HandlePictureDecode(CUVIDPICPARAMS *PicParams);
+    int HandlePictureDisplay(CUVIDPARSERDISPINFO *DispInfo);
+
+private:
+	unsigned long GetNumDecodeSurfaces(cudaVideoCodec Codec, unsigned int Width, unsigned int Height);
 };
+
+namespace Internal 
+{
+    int HandleVideoSequenceCallback(void *UserData, CUVIDEOFORMAT *VideoFormat);
+	int HandlePictureDecodeCallback(void *UserData, CUVIDPICPARAMS *PicParams);
+	int HandlePictureDisplayCallback(void *UserData, CUVIDPARSERDISPINFO *DispInfo);
+}
