Fix build error after driver version struct refactor.

diff --git a/Engine/Plugins/Experimental/AVCodecs/NVCodecs/Source/NVCodecsRHI/Private/NVCodecRHIModule.cpp b/Engine/Plugins/Experimental/AVCodecs/NVCodecs/Source/NVCodecsRHI/Private/NVCodecRHIModule.cpp
--- a/Engine/Plugins/Experimental/AVCodecs/NVCodecs/Source/NVCodecsRHI/Private/NVCodecRHIModule.cpp
+++ b/Engine/Plugins/Experimental/AVCodecs/NVCodecs/Source/NVCodecsRHI/Private/NVCodecRHIModule.cpp
@@ -26,7 +26,7 @@ public:
 					RequiredDriverVersion = TEXT("530.41");
 #endif
 
-					bool bExceedsMinimumDriverVersion = FMultiInt<2>(*GRHIAdapterUserDriverVersion) >= FMultiInt<2>(*RequiredDriverVersion);
+					bool bExceedsMinimumDriverVersion = FDriverVersion(GRHIAdapterUserDriverVersion) >= FDriverVersion(RequiredDriverVersion);
 					if(!bExceedsMinimumDriverVersion)
 					{
 						FAVResult::Log(EAVResult::Error, FString::Printf(TEXT("Detected driver version (%s) is older than required (%s). Please update your drivers!"), *GRHIAdapterUserDriverVersion, *RequiredDriverVersion));
