Fix: NVDec causing crashes in containers on module startup

diff --git a/Engine/Plugins/Experimental/AVCodecs/NVCodecs/Source/NVDEC/Private/NVDEC.cpp b/Engine/Plugins/Experimental/AVCodecs/NVCodecs/Source/NVDEC/Private/NVDEC.cpp
--- a/Engine/Plugins/Experimental/AVCodecs/NVCodecs/Source/NVDEC/Private/NVDEC.cpp
+++ b/Engine/Plugins/Experimental/AVCodecs/NVCodecs/Source/NVDEC/Private/NVDEC.cpp
@@ -40,7 +40,7 @@ FNVDEC::FNVDEC()
 #if PLATFORM_WINDOWS
 	static TCHAR const* DllName = TEXT("nvcuvid.dll");
 #elif PLATFORM_LINUX
-	static TCHAR const* DllName = TEXT("libnvcuvid.so");
+	static TCHAR const* DllName = TEXT("libnvcuvid.so.1");
 #else
 	static TCHAR const* DllName = nullptr;
 #endif
@@ -70,6 +70,14 @@ FNVDEC::FNVDEC()
 			LOAD_FROM_DLL(cuvidUnmapVideoFrame);
 #endif
 		}
+		else
+		{
+			FAVResult::Log(EAVResult::Warning, TEXT("Failed to get NVDEC dll handle. NVDEC module will not be available."), TEXT("NVDEC"));
+		}
+	}
+	else
+	{
+		FAVResult::Log(EAVResult::Warning, TEXT("Failed to get NVDEC dll name. NVDEC module will not be available."), TEXT("NVDEC"));
 	}
 }
 
diff --git a/Engine/Plugins/Experimental/AVCodecs/NVCodecs/Source/NVDEC/Private/NVDECModule.cpp b/Engine/Plugins/Experimental/AVCodecs/NVCodecs/Source/NVDEC/Private/NVDECModule.cpp
--- a/Engine/Plugins/Experimental/AVCodecs/NVCodecs/Source/NVDEC/Private/NVDECModule.cpp
+++ b/Engine/Plugins/Experimental/AVCodecs/NVCodecs/Source/NVDEC/Private/NVDECModule.cpp
@@ -16,7 +16,7 @@ public:
 	virtual void StartupModule() override
 	{
 		TFunction<bool(TSharedRef<FAVDevice> const&, cudaVideoCodec)> CheckCodecSupport = [](TSharedRef<FAVDevice> const& Device, cudaVideoCodec Codec) {
-			if(!Device->HasContext<FVideoContextCUDA>())
+			if(!Device->HasContext<FVideoContextCUDA>() || !FAPI::Get<FNVDEC>().IsValid())
 			{
 				return false;
 			}
diff --git a/Engine/Plugins/Experimental/AVCodecs/NVCodecs/Source/NVENC/Private/NVENC.cpp b/Engine/Plugins/Experimental/AVCodecs/NVCodecs/Source/NVENC/Private/NVENC.cpp
--- a/Engine/Plugins/Experimental/AVCodecs/NVCodecs/Source/NVENC/Private/NVENC.cpp
+++ b/Engine/Plugins/Experimental/AVCodecs/NVCodecs/Source/NVENC/Private/NVENC.cpp
@@ -6,6 +6,8 @@
 #include "HAL/Thread.h"
 #include "RHI.h"
 
+#include "AVResult.h"
+
 REGISTER_TYPEID(FNVENC);
 
 // define a function pointer for creating an instance of nvEncodeAPI
@@ -55,6 +57,14 @@ FNVENC::FNVENC()
 
 			DllHandle = nullptr;
 		}
+		else
+		{
+			FAVResult::Log(EAVResult::Warning, TEXT("Failed to get NVENC dll handle. NVENC module will not be available."), TEXT("NVENC"));
+		}
+	}
+	else
+	{
+		FAVResult::Log(EAVResult::Warning, TEXT("Failed to get NVENC dll name. NVENC module will not be available."), TEXT("NVENC"));
 	}
 }
 
diff --git a/Engine/Plugins/Experimental/AVCodecs/NVCodecs/Source/NVENC/Private/NVENCModule.cpp b/Engine/Plugins/Experimental/AVCodecs/NVCodecs/Source/NVENC/Private/NVENCModule.cpp
--- a/Engine/Plugins/Experimental/AVCodecs/NVCodecs/Source/NVENC/Private/NVENCModule.cpp
+++ b/Engine/Plugins/Experimental/AVCodecs/NVCodecs/Source/NVENC/Private/NVENCModule.cpp
@@ -21,6 +21,11 @@ public:
 				return false;
 			}
 
+			if(!FAPI::Get<FNVENC>().IsValid())
+			{
+				return false;
+			}
+
 			NV_ENC_STRUCT(NV_ENC_OPEN_ENCODE_SESSION_EX_PARAMS, SessionParams);
 			SessionParams.apiVersion = NVENCAPI_VERSION;
 			SessionParams.deviceType = NV_ENC_DEVICE_TYPE_CUDA;
