Fix: Don't run NVidia driver check on non-NVidia GPUs (facepalm)

diff --git a/Engine/Plugins/Experimental/AVCodecs/NVCodecs/Source/NVCodecsRHI/Private/NVCodecRHIModule.cpp b/Engine/Plugins/Experimental/AVCodecs/NVCodecs/Source/NVCodecsRHI/Private/NVCodecRHIModule.cpp
--- a/Engine/Plugins/Experimental/AVCodecs/NVCodecs/Source/NVCodecsRHI/Private/NVCodecRHIModule.cpp
+++ b/Engine/Plugins/Experimental/AVCodecs/NVCodecs/Source/NVCodecsRHI/Private/NVCodecRHIModule.cpp
@@ -17,7 +17,9 @@ public:
 	{
 		if (FApp::CanEverRender())
 		{
-			FCoreDelegates::OnPostEngineInit.AddLambda([]()
+			FCoreDelegates::OnPostEngineInit.AddLambda([]() {
+				bool bExceedsMinimumDriverVersion = false;
+				if (IsRHIDeviceNVIDIA())
 				{
 					FString RequiredDriverVersion;
 #if PLATFORM_WINDOWS
@@ -26,15 +28,16 @@ public:
 					RequiredDriverVersion = TEXT("530.41");
 #endif
 
-					bool bExceedsMinimumDriverVersion = FDriverVersion(GRHIAdapterUserDriverVersion) >= FDriverVersion(RequiredDriverVersion);
-					if(!bExceedsMinimumDriverVersion)
+					bExceedsMinimumDriverVersion = FDriverVersion(GRHIAdapterUserDriverVersion) >= FDriverVersion(RequiredDriverVersion);
+					if (!bExceedsMinimumDriverVersion)
 					{
 						FAVResult::Log(EAVResult::Error, FString::Printf(TEXT("Detected driver version (%s) is older than required (%s). Please update your drivers!"), *GRHIAdapterUserDriverVersion, *RequiredDriverVersion));
 					}
+				}
 
-					const_cast<FNVDEC&>(FAPI::Get<FNVDEC>()).bHasCompatibleGPU = IsRHIDeviceNVIDIA() && bExceedsMinimumDriverVersion;
-					const_cast<FNVENC&>(FAPI::Get<FNVENC>()).bHasCompatibleGPU = IsRHIDeviceNVIDIA() && bExceedsMinimumDriverVersion;
-				});
+				const_cast<FNVDEC&>(FAPI::Get<FNVDEC>()).bHasCompatibleGPU = IsRHIDeviceNVIDIA() && bExceedsMinimumDriverVersion;
+				const_cast<FNVENC&>(FAPI::Get<FNVENC>()).bHasCompatibleGPU = IsRHIDeviceNVIDIA() && bExceedsMinimumDriverVersion;
+			});
 		}
 	}
 };
