Fix AVCodecs NVDEC crash when video resizes larger than original size. Problem occured because the max resolution was set to the original video size, not the max supporrted size by the device.

diff --git a/Engine/Plugins/Experimental/AVCodecs/NVCodecs/Source/NVDEC/Private/Video/Decoders/VideoDecoderNVDEC.cpp b/Engine/Plugins/Experimental/AVCodecs/NVCodecs/Source/NVDEC/Private/Video/Decoders/VideoDecoderNVDEC.cpp
--- a/Engine/Plugins/Experimental/AVCodecs/NVCodecs/Source/NVDEC/Private/Video/Decoders/VideoDecoderNVDEC.cpp
+++ b/Engine/Plugins/Experimental/AVCodecs/NVCodecs/Source/NVDEC/Private/Video/Decoders/VideoDecoderNVDEC.cpp
@@ -9,7 +9,7 @@
 
 namespace Internal 
 {
-    int HandleVideoSequenceCallback(void* UserData, CUVIDEOFORMAT* VideoFormat)
+	int HandleVideoSequenceCallback(void* UserData, CUVIDEOFORMAT* VideoFormat)
 	{
 		return static_cast<FVideoDecoderNVDEC*>(UserData)->HandleVideoSequence(VideoFormat);
 	}
@@ -45,14 +45,14 @@ FAVResult FVideoDecoderNVDEC::Open(TSharedRef<FAVDevice> const& NewDevice, TShar
 
 	CUVIDPARSERPARAMS VideoParserParameters = {};
 	FVideoDecoderConfigNVDEC const& PendingConfig = GetPendingConfig();
-    VideoParserParameters.CodecType = PendingConfig.CodecType;
-    VideoParserParameters.ulMaxNumDecodeSurfaces = 1;
+	VideoParserParameters.CodecType = PendingConfig.CodecType;
+	VideoParserParameters.ulMaxNumDecodeSurfaces = 1;
 	// TODO (william.belcher): bLowLatency ? 0 : 1;
-    VideoParserParameters.ulMaxDisplayDelay = 0;
-    VideoParserParameters.pUserData = this;
-    VideoParserParameters.pfnSequenceCallback = Internal::HandleVideoSequenceCallback;
-    VideoParserParameters.pfnDecodePicture = Internal::HandlePictureDCodecallback;
-    VideoParserParameters.pfnDisplayPicture = Internal::HandlePictureDisplayCallback;
+	VideoParserParameters.ulMaxDisplayDelay = 0;
+	VideoParserParameters.pUserData = this;
+	VideoParserParameters.pfnSequenceCallback = Internal::HandleVideoSequenceCallback;
+	VideoParserParameters.pfnDecodePicture = Internal::HandlePictureDCodecallback;
+	VideoParserParameters.pfnDisplayPicture = Internal::HandlePictureDisplayCallback;
 
 	if (FAPI::Get<FNVDEC>().cuvidCreateVideoParser(&Parser, &VideoParserParameters) != CUDA_SUCCESS)
 	{
@@ -114,6 +114,10 @@ bool FVideoDecoderNVDEC::IsInitialized() const
 	return IsOpen() && Decoder != nullptr;
 }
 
+// ApplyConfig is different from ApplyConfig from other decoder classes in that it does not call Close on an error.
+// This is because ApplyConfig is called from a function chain originating from cuvidParseVideoData.
+// Because cuvidParseVideoData depends on the Parser object, ApplyConfig cannot call Close which destroys the Parser.
+// Destroying Parser while in use will crash from within Nvidia dll.
 FAVResult FVideoDecoderNVDEC::ApplyConfig()
 {
 	if (IsOpen())
@@ -134,21 +138,20 @@ FAVResult FVideoDecoderNVDEC::ApplyConfig()
 						CUVIDRECONFIGUREDECODERINFO ReconfigParams = { 0 };
 
 						ReconfigParams.ulWidth = PendingConfig.ulWidth;
-   						ReconfigParams.ulHeight = PendingConfig.ulHeight;
-
-    					ReconfigParams.display_area.bottom = PendingConfig.display_area.bottom;
-    					ReconfigParams.display_area.top = PendingConfig.display_area.top;
-    					ReconfigParams.display_area.left = PendingConfig.display_area.left;
-    					ReconfigParams.display_area.right = PendingConfig.display_area.right;
-    					ReconfigParams.ulTargetWidth = PendingConfig.ulWidth;
-    					ReconfigParams.ulTargetHeight = PendingConfig.ulHeight;
+						ReconfigParams.ulHeight = PendingConfig.ulHeight;
+
+						ReconfigParams.display_area.bottom = PendingConfig.display_area.bottom;
+						ReconfigParams.display_area.top = PendingConfig.display_area.top;
+						ReconfigParams.display_area.left = PendingConfig.display_area.left;
+						ReconfigParams.display_area.right = PendingConfig.display_area.right;
+						ReconfigParams.ulTargetWidth = PendingConfig.ulWidth;
+						ReconfigParams.ulTargetHeight = PendingConfig.ulHeight;
 						ReconfigParams.ulNumDecodeSurfaces = PendingConfig.ulNumDecodeSurfaces;
 
 						FCUDAContextScope const ContextGuard(GetDevice()->GetContext<FVideoContextCUDA>()->Raw);
 						CUresult const Result = FAPI::Get<FNVDEC>().cuvidReconfigureDecoder(Decoder, &ReconfigParams);
 						if (Result != CUDA_SUCCESS)
 						{
-							Close();
 							return FAVResult(EAVResult::Error, TEXT("Failed to reconfigure decoder"), TEXT("NVDEC"), Result);
 						}
 					}
@@ -159,7 +162,6 @@ FAVResult FVideoDecoderNVDEC::ApplyConfig()
 					CUresult const Result = FAPI::Get<FNVDEC>().cuvidDestroyDecoder(Decoder);
 					if (Result != CUDA_SUCCESS)
 					{
-						Close();
 						return FAVResult(EAVResult::ErrorDestroying, TEXT("Failed to destroy NVDEC decoder"), TEXT("NVDEC"), Result);
 					}
 
@@ -173,7 +175,6 @@ FAVResult FVideoDecoderNVDEC::ApplyConfig()
 				CUresult const Result = FAPI::Get<FNVDEC>().cuvidCreateDecoder(&Decoder, &const_cast<FVideoDecoderConfigNVDEC&>(PendingConfig));
 				if (Result != CUDA_SUCCESS)
 				{
-					Close();
 					return FAVResult(EAVResult::ErrorCreating, TEXT("Failed to create decoder"), TEXT("NVDEC"), Result);
 				}
 			}
@@ -213,19 +214,19 @@ FAVResult FVideoDecoderNVDEC::SendPacket(FVideoPacket const& Packet)
 	}
 
 	CUVIDSOURCEDATAPACKET CUPacket = {0};
-    CUPacket.payload = Packet.DataPtr.Get();
-    CUPacket.payload_size = Packet.DataSize;
-    CUPacket.flags = CUVID_PKT_TIMESTAMP;
+	CUPacket.payload = Packet.DataPtr.Get();
+	CUPacket.payload_size = Packet.DataSize;
+	CUPacket.flags = CUVID_PKT_TIMESTAMP;
 	// TODO (william.belcher): Adding this flag assumes that each Packet contains a full frame, but is required
 	// to prevent the parser from adding a 1 frame delay
 	CUPacket.flags |= CUVID_PKT_ENDOFPICTURE;
-    CUPacket.timestamp = 0;
-    if (!Packet.DataPtr.Get() || Packet.DataSize == 0) 
+	CUPacket.timestamp = 0;
+	if (!Packet.DataPtr.Get() || Packet.DataSize == 0) 
 	{
-        CUPacket.flags |= CUVID_PKT_ENDOFSTREAM;
-    }
+		CUPacket.flags |= CUVID_PKT_ENDOFSTREAM;
+	}
 
-    CUresult const CUResult = FAPI::Get<FNVDEC>().cuvidParseVideoData(Parser, &CUPacket);
+	CUresult const CUResult = FAPI::Get<FNVDEC>().cuvidParseVideoData(Parser, &CUPacket);
 	if (CUResult != CUDA_SUCCESS)
 	{
 		return FAVResult(EAVResult::Error, TEXT("Failed to parse video data"), TEXT("NVDEC"), CUResult);
@@ -355,27 +356,31 @@ int FVideoDecoderNVDEC::HandleVideoSequence(CUVIDEOFORMAT *VideoFormat)
 	FVideoDecoderConfigNVDEC& PendingConfig = EditPendingConfig();
 	PendingConfig.CodecType = VideoFormat->codec;
 	PendingConfig.ChromaFormat = VideoFormat->chroma_format;
-    PendingConfig.OutputFormat = VideoFormat->bit_depth_luma_minus8 ? cudaVideoSurfaceFormat_P016 : cudaVideoSurfaceFormat_NV12;
-    PendingConfig.bitDepthMinus8 = VideoFormat->bit_depth_luma_minus8;
-    PendingConfig.DeinterlaceMode = cudaVideoDeinterlaceMode_Weave;
-    PendingConfig.ulNumOutputSurfaces = 2;
-    // With PreferCUVID, JPEG is still decoded by CUDA while video is decoded by NVDEC hardware
-    PendingConfig.ulCreationFlags = cudaVideoCreate_PreferCUVID;
-    PendingConfig.ulNumDecodeSurfaces = NumDecodeSurfaces;
-    PendingConfig.vidLock = CtxLock;
-    PendingConfig.ulWidth = VideoFormat->coded_width;
-    PendingConfig.ulHeight = VideoFormat->coded_height;
-    PendingConfig.ulMaxWidth = VideoFormat->coded_width;
-    PendingConfig.ulMaxHeight = VideoFormat->coded_height;
+	PendingConfig.OutputFormat = VideoFormat->bit_depth_luma_minus8 ? cudaVideoSurfaceFormat_P016 : cudaVideoSurfaceFormat_NV12;
+	PendingConfig.bitDepthMinus8 = VideoFormat->bit_depth_luma_minus8;
+	PendingConfig.DeinterlaceMode = cudaVideoDeinterlaceMode_Weave;
+	PendingConfig.ulNumOutputSurfaces = 2;
+	// With PreferCUVID, JPEG is still decoded by CUDA while video is decoded by NVDEC hardware
+	PendingConfig.ulCreationFlags = cudaVideoCreate_PreferCUVID;
+	PendingConfig.ulNumDecodeSurfaces = NumDecodeSurfaces;
+	PendingConfig.vidLock = CtxLock;
+	PendingConfig.ulWidth = VideoFormat->coded_width;
+	PendingConfig.ulHeight = VideoFormat->coded_height;
+	// TODO (Eden.Harris) By setting max width/height to the device max, the encoder may potentially be allocating multiple large frames in VRAM.
+	// This may be alot when b-frames are enabled and should be monitored to make sure VRAM is not unnecessarily consumed.
+	// The reason why this is set to the device max is because if the stream resizes, it cannot resize above this initial ulMaxWidth/ulMaxHeight.
+	PendingConfig.ulMaxWidth = DecodeCaps.nMaxWidth;
+	PendingConfig.ulMaxHeight = DecodeCaps.nMaxHeight;
 
 	// TODO (william.belcher): Add support for cropping and resizing
 	PendingConfig.ulTargetWidth = VideoFormat->coded_width;
-    PendingConfig.ulTargetHeight = VideoFormat->coded_height;
+	PendingConfig.ulTargetHeight = VideoFormat->coded_height;
 
 	FAVResult Result = ApplyConfig();
 	if (Result.IsNotSuccess())
 	{
 		FAVResult::Log(EAVResult::Error, TEXT("Failed to apply decoder config"), TEXT("NVDEC"), Result);
+		return 0;	// Return fail. Unfortunately this does not bubble up and HandlePictureDecode will still be called.
 	}
 
 	return NumDecodeSurfaces;
@@ -387,7 +392,7 @@ int FVideoDecoderNVDEC::HandlePictureDecode(CUVIDPICPARAMS *PicParams)
 	if (Result != CUDA_SUCCESS)
 	{
 		FAVResult::Log(EAVResult::Error, TEXT("Failed to decode frame"), TEXT("NVDEC"), Result);
-		return 0;
+		return 0;	// Unfortunately this does not bubble up to cuvidParseVideoData
 	}
 	return 1;
 }
@@ -395,11 +400,11 @@ int FVideoDecoderNVDEC::HandlePictureDecode(CUVIDPICPARAMS *PicParams)
 int FVideoDecoderNVDEC::HandlePictureDisplay(CUVIDPARSERDISPINFO *DispInfo)
 {
 	CUVIDPROCPARAMS VideoProcessingParameters = {};
-    VideoProcessingParameters.progressive_frame = DispInfo->progressive_frame;
-    VideoProcessingParameters.second_field = DispInfo->repeat_first_field + 1;
-    VideoProcessingParameters.top_field_first = DispInfo->top_field_first;
-    VideoProcessingParameters.unpaired_field = DispInfo->repeat_first_field < 0;
-    VideoProcessingParameters.output_stream = 0;
+	VideoProcessingParameters.progressive_frame = DispInfo->progressive_frame;
+	VideoProcessingParameters.second_field = DispInfo->repeat_first_field + 1;
+	VideoProcessingParameters.top_field_first = DispInfo->top_field_first;
+	VideoProcessingParameters.unpaired_field = DispInfo->repeat_first_field < 0;
+	VideoProcessingParameters.output_stream = 0;
 
 	Frames.Enqueue({ DispInfo->picture_index, AppliedConfig.ulTargetWidth, AppliedConfig.ulTargetHeight, GetAppliedConfig().OutputFormat, VideoProcessingParameters });
 	++FramesCount;
