Fix NVENC AVCodecs crash when video resize larger than max device resolution. Code was hitting an unimplemented tag.

NVENC ApplyConfig now checks the device max resolution and will return an error if the frame resolution exceeds this.

When NVENC reconfigures now, it will attempt to re-open and initialize the device. The default max resolution is now the smaller of the device max resolution and frame resolution to ensure vram is not wasted and to make sure the device max resolution is not exceeded. This does mean NVENC will need to re-initialize when the frame resolution increases.

This issue was previously not hit because PixelStreaming will reset NVENC on resolution change so reconfigure code path was not hit.

diff --git a/Engine/Plugins/Experimental/AVCodecs/NVCodecs/Source/NVENC/Private/Video/Encoders/VideoEncoderNVENC.cpp b/Engine/Plugins/Experimental/AVCodecs/NVCodecs/Source/NVENC/Private/Video/Encoders/VideoEncoderNVENC.cpp
--- a/Engine/Plugins/Experimental/AVCodecs/NVCodecs/Source/NVENC/Private/Video/Encoders/VideoEncoderNVENC.cpp
+++ b/Engine/Plugins/Experimental/AVCodecs/NVCodecs/Source/NVENC/Private/Video/Encoders/VideoEncoderNVENC.cpp
@@ -12,6 +12,27 @@ bool FEncoderNVENC::IsOpen() const
 	return Encoder != nullptr;
 }
 
+FAVResult FEncoderNVENC::ReOpen()
+{
+	Close();
+
+	NV_ENC_STRUCT(NV_ENC_OPEN_ENCODE_SESSION_EX_PARAMS, SessionParams);
+	SessionParams.apiVersion = NVENCAPI_VERSION;
+
+	SessionParams.deviceType = SessionDeviceType;
+	SessionParams.device = SessionDevice;
+
+	NVENCSTATUS const Result = FAPI::Get<FNVENC>().nvEncOpenEncodeSessionEx(&const_cast<NV_ENC_OPEN_ENCODE_SESSION_EX_PARAMS&>(SessionParams), &Encoder);
+	if (Result != NV_ENC_SUCCESS)
+	{
+		Close();
+
+		return FAVResult(EAVResult::ErrorCreating, TEXT("Failed to re-create encoder"), TEXT("NVENC"), Result);
+	}
+
+	return EAVResult::Success;
+}
+
 FAVResult FEncoderNVENC::Open(TSharedRef<FAVDevice> const& NewDevice, TSharedRef<FAVInstance> const& NewInstance, TFunction<void(NV_ENC_OPEN_ENCODE_SESSION_EX_PARAMS&)> SetupEncoderSessionFunc)
 {
 	Close();
@@ -21,6 +42,9 @@ FAVResult FEncoderNVENC::Open(TSharedRef<FAVDevice> const& NewDevice, TSharedRef
 
 	SetupEncoderSessionFunc(SessionParams);
 
+	SessionDeviceType = SessionParams.deviceType;
+	SessionDevice = SessionParams.device;
+
 	NVENCSTATUS const Result = FAPI::Get<FNVENC>().nvEncOpenEncodeSessionEx(&const_cast<NV_ENC_OPEN_ENCODE_SESSION_EX_PARAMS&>(SessionParams), &Encoder);
 	if (Result != NV_ENC_SUCCESS)
 	{
@@ -65,23 +89,39 @@ bool FEncoderNVENC::IsInitialized() const
 	return IsOpen() && Buffer != nullptr;
 }
 
-FAVResult FEncoderNVENC::ApplyConfig(FVideoEncoderConfigNVENC const& AppliedConfig, FVideoEncoderConfigNVENC const& PendingConfig, TFunction<FAVResult()> ApplyConfigFunc)
+FAVResult FEncoderNVENC::ApplyConfig(FVideoEncoderConfigNVENC const& AppliedConfig, FVideoEncoderConfigNVENC& PendingConfig, TFunction<FAVResult()> ApplyConfigFunc)
 {
 	if (IsOpen())
 	{
 		// FVideoEncoderConfigNVENC const& PendingConfig = GetPendingConfig();
 		if (AppliedConfig != PendingConfig)
 		{
+			// Because TransformConfig does not have access to MaxDeviceEncodeWidth/Height, we check this here and re-adjust the PendingConfig maxEncodeWidth/Height if it goes over.
+			SetMaxResolution(PendingConfig);
+
+			if (PendingConfig.encodeWidth > MaxDeviceEncodeWidth
+				|| PendingConfig.encodeHeight > MaxDeviceEncodeHeight)
+			{
+				return FAVResult(EAVResult::Error, TEXT("Encoder pending resolution exceeds device max capabilities"));
+			}
+			if (PendingConfig.encodeWidth > PendingConfig.maxEncodeWidth
+				|| PendingConfig.encodeHeight > PendingConfig.maxEncodeHeight)
+			{
+				return FAVResult(EAVResult::Error, TEXT("Encoder pending resolution exceeds configuration max resolution"));
+			}
+
 			if (IsInitialized())
 			{
 				// Can be reconfigured? See https://docs.nvidia.com/video-technologies/video-codec-sdk/nvenc-video-encoder-api-prog-guide/#reconfigure-api
-				if (AppliedConfig.maxEncodeWidth == PendingConfig.maxEncodeWidth
-					&& AppliedConfig.maxEncodeHeight == PendingConfig.maxEncodeHeight
-					&& AppliedConfig.enablePTD == PendingConfig.enablePTD
+				if (AppliedConfig.enablePTD == PendingConfig.enablePTD
 					&& AppliedConfig.enableEncodeAsync == PendingConfig.enableEncodeAsync
 					&& AppliedConfig.encodeConfig->gopLength == PendingConfig.encodeConfig->gopLength
 					&& AppliedConfig.encodeConfig->frameIntervalP == PendingConfig.encodeConfig->frameIntervalP
-					&& AppliedConfig.encodeConfig->encodeCodecConfig.h264Config.idrPeriod == PendingConfig.encodeConfig->encodeCodecConfig.h264Config.idrPeriod)
+					&& AppliedConfig.encodeConfig->encodeCodecConfig.h264Config.idrPeriod == PendingConfig.encodeConfig->encodeCodecConfig.h264Config.idrPeriod
+					&& AppliedConfig.maxEncodeWidth == PendingConfig.maxEncodeWidth
+					&& AppliedConfig.maxEncodeHeight == PendingConfig.maxEncodeHeight
+					&& PendingConfig.encodeWidth <= PendingConfig.maxEncodeWidth
+					&& PendingConfig.encodeHeight <= PendingConfig.maxEncodeHeight)
 				{
 					NV_ENC_STRUCT(NV_ENC_RECONFIGURE_PARAMS, ReconfigureParams);
 					FMemory::Memcpy(&ReconfigureParams.reInitEncodeParams, &static_cast<NV_ENC_INITIALIZE_PARAMS const&>(PendingConfig), sizeof(NV_ENC_INITIALIZE_PARAMS));
@@ -95,8 +135,11 @@ FAVResult FEncoderNVENC::ApplyConfig(FVideoEncoderConfigNVENC const& AppliedConf
 				}
 				else
 				{
-					// TODO: Destroy and recreate with original session
-					unimplemented();
+					FAVResult const Result = ReOpen();
+					if (Result.IsNotSuccess())
+					{
+						return Result;
+					}
 				}
 			}
 
@@ -149,6 +192,30 @@ int FEncoderNVENC::GetCapability(GUID EncodeGUID, NV_ENC_CAPS CapsToQuery) const
 	return FAVResult(EAVResult::ErrorInvalidState, TEXT("Encoder not open"), TEXT("NVENC"));
 }
 
+void FEncoderNVENC::GetMaxDeviceEncodeResolution(const FVideoEncoderConfigNVENC& PendingConfig)
+{
+	if (IsOpen())
+	{
+		if (!bHasMaxDeviceResolution)
+		{
+			MaxDeviceEncodeWidth = static_cast<uint32>(GetCapability(PendingConfig.encodeGUID, NV_ENC_CAPS_WIDTH_MAX));
+			MaxDeviceEncodeHeight = static_cast<uint32>(GetCapability(PendingConfig.encodeGUID, NV_ENC_CAPS_HEIGHT_MAX));
+			bHasMaxDeviceResolution = true;
+		}
+	}
+}
+
+void FEncoderNVENC::SetMaxResolution(FVideoEncoderConfigNVENC& PendingConfig)
+{
+	if (IsOpen())
+	{
+		GetMaxDeviceEncodeResolution(PendingConfig);
+
+		PendingConfig.maxEncodeWidth = FMath::Min(PendingConfig.maxEncodeWidth, MaxDeviceEncodeWidth);
+		PendingConfig.maxEncodeHeight = FMath::Min(PendingConfig.maxEncodeHeight, MaxDeviceEncodeHeight);
+	}
+}
+
 FAVResult FEncoderNVENC::SendFrame(TSharedPtr<FVideoResource> const& Resource, uint32 Timestamp, bool bForceKeyframe, TFunction<FAVResult()> ApplyConfigFunc, TFunction<void(NV_ENC_REGISTER_RESOURCE&)> SetResourceToRegisterFunc)
 {
 	if (IsOpen())
@@ -354,7 +421,7 @@ FAVResult FEncoderNVENC::SendFrameD3D11(TRefCountPtr<ID3D11Device> Device, TShar
 const FString AVGetComErrorDescription(HRESULT Res)
 {
 	const uint32 BufSize = 4096;
-	WIDECHAR Buffer[4096];
+	WIDECHAR	 Buffer[4096];
 	if (::FormatMessageW(FORMAT_MESSAGE_FROM_SYSTEM,
 			nullptr,
 			Res,
@@ -380,7 +447,7 @@ const FString AVGetComErrorDescription(HRESULT Res)
  */
 FAVResult FEncoderNVENC::CreateD3D11Device(TSharedRef<FAVDevice> const& InDevice, TRefCountPtr<ID3D11Device>& OutEncoderDevice, TRefCountPtr<ID3D11DeviceContext>& OutEncoderDeviceContext)
 {
-	TRefCountPtr<IDXGIDevice> DXGIDevice;
+	TRefCountPtr<IDXGIDevice>  DXGIDevice;
 	TRefCountPtr<IDXGIAdapter> Adapter;
 
 	HRESULT Result = InDevice->GetContext<FVideoContextD3D11>()->Device->QueryInterface(__uuidof(IDXGIDevice), (void**)DXGIDevice.GetInitReference());
@@ -393,7 +460,7 @@ FAVResult FEncoderNVENC::CreateD3D11Device(TSharedRef<FAVDevice> const& InDevice
 		return FAVResult(EAVResult::Fatal, FString::Printf(TEXT("DXGIDevice::GetAdapter() failed 0x%X - %s."), (uint32)Result, *AVGetComErrorDescription(Result)), TEXT("D3D11"));
 	}
 
-	uint32 DeviceFlags = 0;
+	uint32			  DeviceFlags = 0;
 	D3D_FEATURE_LEVEL FeatureLevel = D3D_FEATURE_LEVEL_11_0;
 	D3D_FEATURE_LEVEL ActualFeatureLevel;
 
diff --git a/Engine/Plugins/Experimental/AVCodecs/NVCodecs/Source/NVENC/Public/Video/Encoders/Configs/VideoEncoderConfigNVENC.h b/Engine/Plugins/Experimental/AVCodecs/NVCodecs/Source/NVENC/Public/Video/Encoders/Configs/VideoEncoderConfigNVENC.h
--- a/Engine/Plugins/Experimental/AVCodecs/NVCodecs/Source/NVENC/Public/Video/Encoders/Configs/VideoEncoderConfigNVENC.h
+++ b/Engine/Plugins/Experimental/AVCodecs/NVCodecs/Source/NVENC/Public/Video/Encoders/Configs/VideoEncoderConfigNVENC.h
@@ -31,8 +31,8 @@ public:
 		enablePTD = 1;
 		reportSliceOffsets = 0;
 		enableSubFrameWrite = 0;
-		maxEncodeWidth = 4096;
-		maxEncodeHeight = 4096;
+		maxEncodeWidth = 0;
+		maxEncodeHeight = 0;
 
 		Config.version = NV_ENC_CONFIG_VER;
 		encodeConfig = &Config;
diff --git a/Engine/Plugins/Experimental/AVCodecs/NVCodecs/Source/NVENC/Public/Video/Encoders/VideoEncoderNVENC.h b/Engine/Plugins/Experimental/AVCodecs/NVCodecs/Source/NVENC/Public/Video/Encoders/VideoEncoderNVENC.h
--- a/Engine/Plugins/Experimental/AVCodecs/NVCodecs/Source/NVENC/Public/Video/Encoders/VideoEncoderNVENC.h
+++ b/Engine/Plugins/Experimental/AVCodecs/NVCodecs/Source/NVENC/Public/Video/Encoders/VideoEncoderNVENC.h
@@ -25,16 +25,27 @@ private:
 	NV_ENC_OUTPUT_PTR Buffer = nullptr;
 	TQueue<FVideoPacket> Packets;
 
+	NV_ENC_DEVICE_TYPE  SessionDeviceType;
+	void* SessionDevice;
+
+	uint32 MaxDeviceEncodeWidth = 0;
+	uint32 MaxDeviceEncodeHeight = 0;
+	bool bHasMaxDeviceResolution = false;
+
+	FAVResult ReOpen();
+	void SetMaxResolution(FVideoEncoderConfigNVENC& PendingConfig);
+	void GetMaxDeviceEncodeResolution(const FVideoEncoderConfigNVENC& PendingConfig);
+
 public:
 	// Begin matching the TVideoEncoder interface
 	virtual ~FEncoderNVENC();
 	virtual bool IsOpen() const;
 	virtual FAVResult Open(TSharedRef<FAVDevice> const& NewDevice, TSharedRef<FAVInstance> const& NewInstance, TFunction<void(NV_ENC_OPEN_ENCODE_SESSION_EX_PARAMS&)> SetupEncoderSessionFunc);
 	virtual void Close();
-	virtual FAVResult ApplyConfig(FVideoEncoderConfigNVENC const& AppliedConfig, FVideoEncoderConfigNVENC const& PendingConfig, TFunction<FAVResult()> ApplyConfigFunc);
+	virtual FAVResult ApplyConfig(FVideoEncoderConfigNVENC const& AppliedConfig, FVideoEncoderConfigNVENC& PendingConfig, TFunction<FAVResult()> ApplyConfigFunc);
 #if PLATFORM_WINDOWS
 	virtual FAVResult CreateD3D11Device(TSharedRef<FAVDevice> const& InDevice, TRefCountPtr<ID3D11Device>& OutEncoderDevice, TRefCountPtr<ID3D11DeviceContext>& OutEncoderDeviceContext);
-	virtual FAVResult SendFrameD3D11(TRefCountPtr<ID3D11Device> Device, TSharedPtr<FVideoResourceD3D11> const& Resource, uint32 Timestamp, bool bForceKeyframe, TFunction<FAVResult()> ApplyConfigFunc);
+	virtual FAVResult SendFrameD3D11( TRefCountPtr<ID3D11Device> Device, TSharedPtr<FVideoResourceD3D11> const& Resource, uint32 Timestamp, bool bForceKeyframe, TFunction<FAVResult()> ApplyConfigFunc);
 #endif // PLATFORM_WINDOWS
 	virtual FAVResult SendFrameCUDA(TSharedPtr<FVideoResourceCUDA> const& Resource, uint32 Timestamp, bool bForceKeyframe, TFunction<FAVResult()> ApplyConfigFunc);
 	virtual FAVResult SendFrame(TSharedPtr<FVideoResource> const& Resource, uint32 Timestamp, bool bForceKeyframe, TFunction<FAVResult()> ApplyConfigFunc, TFunction<void(NV_ENC_REGISTER_RESOURCE&)> SetResourceToRegisterFunc);
@@ -66,7 +77,7 @@ public:
 
 	virtual FAVResult ApplyConfig() override
 	{
-		return Base->ApplyConfig(AppliedConfig, GetPendingConfig(), [this]() {
+		return Base->ApplyConfig(AppliedConfig, EditPendingConfig(), [this]() {
 			return TVideoEncoder::ApplyConfig();
 		});
 	}
@@ -113,7 +124,7 @@ public:
 
 	virtual FAVResult ApplyConfig() override
 	{
-		return Base->ApplyConfig(AppliedConfig, GetPendingConfig(), [this]() {
+		return Base->ApplyConfig(AppliedConfig, EditPendingConfig(), [this]() {
 			return TVideoEncoder::ApplyConfig();
 		});
 	};
