Updating AVCodecs to allow Enums to be exposed in Editor Menu.
Switching QP for Quality in AVCodecs and updating usage in PixelStreaming.

#rnx

diff --git a/Engine/Plugins/Experimental/AVCodecs/AMFCodecs/Source/AMFCodecs/Private/Video/Encoders/Configs/VideoEncoderConfigAMF.cpp b/Engine/Plugins/Experimental/AVCodecs/AMFCodecs/Source/AMFCodecs/Private/Video/Encoders/Configs/VideoEncoderConfigAMF.cpp
--- a/Engine/Plugins/Experimental/AVCodecs/AMFCodecs/Source/AMFCodecs/Private/Video/Encoders/Configs/VideoEncoderConfigAMF.cpp
+++ b/Engine/Plugins/Experimental/AVCodecs/AMFCodecs/Source/AMFCodecs/Private/Video/Encoders/Configs/VideoEncoderConfigAMF.cpp
@@ -54,7 +54,7 @@ DLLEXPORT FAVResult FAVExtension::TransformConfig(FVideoEncoderConfigAMF& OutCon
 
 	OutConfig.Width = InConfig.Width;
 	OutConfig.Height = InConfig.Height;
-	
+
 	OutConfig.SetProperty(AMF_VIDEO_ENCODER_USAGE, AMF_VIDEO_ENCODER_USAGE_LOW_LATENCY);
 
 	AMFRate const FrameRate = { InConfig.TargetFramerate, 1 };
@@ -66,9 +66,9 @@ DLLEXPORT FAVResult FAVExtension::TransformConfig(FVideoEncoderConfigAMF& OutCon
 	{
 		return ConvertedRateControlMode;
 	}
-	
+
 	OutConfig.SetProperty(AMF_VIDEO_ENCODER_RATE_CONTROL_METHOD, ConvertedRateControlMode.ReturnValue);
-	
+
 	if (InConfig.RateControlMode == ERateControlMode::CBR)
 	{
 		OutConfig.SetProperty(AMF_VIDEO_ENCODER_FILLER_DATA_ENABLE, InConfig.bFillData);
@@ -76,14 +76,14 @@ DLLEXPORT FAVResult FAVExtension::TransformConfig(FVideoEncoderConfigAMF& OutCon
 
 	OutConfig.SetProperty(AMF_VIDEO_ENCODER_PEAK_BITRATE, InConfig.MaxBitrate > -1 ? InConfig.MaxBitrate : DEFAULT_BITRATE_MAX);
 #endif
-	
+
 	OutConfig.SetProperty(AMF_VIDEO_ENCODER_TARGET_BITRATE, InConfig.TargetBitrate > -1 ? InConfig.TargetBitrate : DEFAULT_BITRATE_TARGET);
 
 	OutConfig.SetProperty(AMF_VIDEO_ENCODER_QUALITY_PRESET, AMF_VIDEO_ENCODER_QUALITY_PRESET_QUALITY);
 	OutConfig.SetProperty(AMF_VIDEO_ENCODER_B_PIC_PATTERN, 0);
 
-	OutConfig.SetProperty(AMF_VIDEO_ENCODER_MIN_QP, FMath::Clamp<amf_int64>(InConfig.MinQP, 0, 51));
-	OutConfig.SetProperty(AMF_VIDEO_ENCODER_MAX_QP, InConfig.MaxQP > -1 ? FMath::Clamp<amf_int64>(InConfig.MaxQP, 0, 51) : 51);
+	OutConfig.SetProperty(AMF_VIDEO_ENCODER_MIN_QP, FMath::Floor<amf_int64>((1.0f - (InConfig.MaxQuality / 100.0f)) * 51.0f));
+	OutConfig.SetProperty(AMF_VIDEO_ENCODER_MAX_QP, FMath::Floor<amf_int64>((1.0f - (InConfig.MinQuality / 100.0f)) * 51.0f));
 
 	OutConfig.SetProperty(AMF_VIDEO_ENCODER_QUERY_TIMEOUT, 16);
 
@@ -115,21 +115,26 @@ DLLEXPORT FAVResult FAVExtension::TransformConfig(FVideoEncoderConfig& OutConfig
 	{
 		return ConvertedRateControlMode;
 	}
-	
+
 	OutConfig.RateControlMode = ConvertedRateControlMode;
 
 	bool FillData;
 	InConfig.GetProperty(AMF_VIDEO_ENCODER_FILLER_DATA_ENABLE, &FillData);
-	
+
 	OutConfig.bFillData = FillData;
 
 	InConfig.GetProperty(AMF_VIDEO_ENCODER_PEAK_BITRATE, &OutConfig.MaxBitrate);
 #endif
-	
+
 	InConfig.GetProperty(AMF_VIDEO_ENCODER_TARGET_BITRATE, &OutConfig.TargetBitrate);
 
-	InConfig.GetProperty(AMF_VIDEO_ENCODER_MIN_QP, &OutConfig.MinQP);
-	InConfig.GetProperty(AMF_VIDEO_ENCODER_MAX_QP, &OutConfig.MaxQP);
+	amf_int64 OutMinQpValue;
+	InConfig.GetProperty(AMF_VIDEO_ENCODER_MIN_QP, &OutMinQpValue);
+	OutConfig.MaxQuality = (1.0f - (OutMinQpValue / 51.0f)) * 100.0f;
+
+	amf_int64 OutMaxQpValue;
+	InConfig.GetProperty(AMF_VIDEO_ENCODER_MAX_QP, &OutMaxQpValue);
+	OutConfig.MinQuality = (1.0f - (OutMaxQpValue / 51.0f)) * 100.0f;
 
 	InConfig.GetProperty(AMF_VIDEO_ENCODER_IDR_PERIOD, &OutConfig.KeyframeInterval);
 
@@ -146,23 +151,23 @@ DLLEXPORT FAVResult FAVExtension::TransformConfig(FVideoEncoderConfigAMF& OutCon
 	Config.encodeCodecConfig.h264Config.h264VUIParameters.colourPrimaries = 2;
 	Config.encodeCodecConfig.h264Config.h264VUIParameters.transferCharacteristics = 2;
 	Config.encodeCodecConfig.h264Config.h264VUIParameters.colourMatrix = 2;*/
-	
+
 	static auto const ConvertProfile = [](EH264Profile Profile) -> TAVResult<AMF_VIDEO_ENCODER_PROFILE_ENUM>
-	{
-		switch (Profile)
 		{
-		case EH264Profile::Baseline:
-			return AMF_VIDEO_ENCODER_PROFILE_ENUM::AMF_VIDEO_ENCODER_PROFILE_BASELINE;
-		case EH264Profile::Main:
-			return AMF_VIDEO_ENCODER_PROFILE_ENUM::AMF_VIDEO_ENCODER_PROFILE_MAIN;
-		case EH264Profile::High:
-			return AMF_VIDEO_ENCODER_PROFILE_ENUM::AMF_VIDEO_ENCODER_PROFILE_HIGH;
-		case EH264Profile::ConstrainedHigh:
-			return AMF_VIDEO_ENCODER_PROFILE_ENUM::AMF_VIDEO_ENCODER_PROFILE_CONSTRAINED_HIGH;
-		default:
-			return FAVResult(EAVResult::ErrorUnsupported, FString::Printf(TEXT("H264 profile %d is not supported"), Profile), TEXT("AMF"));
-		}
-	};
+			switch (Profile)
+			{
+			case EH264Profile::Baseline:
+				return AMF_VIDEO_ENCODER_PROFILE_ENUM::AMF_VIDEO_ENCODER_PROFILE_BASELINE;
+			case EH264Profile::Main:
+				return AMF_VIDEO_ENCODER_PROFILE_ENUM::AMF_VIDEO_ENCODER_PROFILE_MAIN;
+			case EH264Profile::High:
+				return AMF_VIDEO_ENCODER_PROFILE_ENUM::AMF_VIDEO_ENCODER_PROFILE_HIGH;
+			case EH264Profile::ConstrainedHigh:
+				return AMF_VIDEO_ENCODER_PROFILE_ENUM::AMF_VIDEO_ENCODER_PROFILE_CONSTRAINED_HIGH;
+			default:
+				return FAVResult(EAVResult::ErrorUnsupported, FString::Printf(TEXT("H264 profile %d is not supported"), Profile), TEXT("AMF"));
+			}
+		};
 
 	OutConfig.CodecType = FVideoEncoderConfigAMF::CodecTypeH264;
 	OutConfig.RepeatSPSPPS = InConfig.RepeatSPSPPS;
@@ -185,15 +190,15 @@ DLLEXPORT FAVResult FAVExtension::TransformConfig(FVideoEncoderConfigAMF& OutCon
 	OutConfig.CodecType = FVideoEncoderConfigAMF::CodecTypeH265;
 
 	static auto const ConvertProfile = [](EH265Profile Profile) -> TAVResult<AMF_VIDEO_ENCODER_PROFILE_ENUM>
-	{
-		switch (Profile)
 		{
-		case EH265Profile::Main:
-			return AMF_VIDEO_ENCODER_PROFILE_ENUM::AMF_VIDEO_ENCODER_PROFILE_MAIN;
-		default:
-			return FAVResult(EAVResult::ErrorUnsupported, FString::Printf(TEXT("H265 profile %d is not supported"), Profile), TEXT("AMF"));
-		}
-	};
+			switch (Profile)
+			{
+			case EH265Profile::Main:
+				return AMF_VIDEO_ENCODER_PROFILE_ENUM::AMF_VIDEO_ENCODER_PROFILE_MAIN;
+			default:
+				return FAVResult(EAVResult::ErrorUnsupported, FString::Printf(TEXT("H265 profile %d is not supported"), Profile), TEXT("AMF"));
+			}
+		};
 
 	OutConfig.CodecType = FVideoEncoderConfigAMF::CodecTypeH265;
 	OutConfig.RepeatSPSPPS = InConfig.RepeatSPSPPS;
diff --git a/Engine/Plugins/Experimental/AVCodecs/AVCodecsCore/Source/AVCodecsCore/AVCodecsCore.Build.cs b/Engine/Plugins/Experimental/AVCodecs/AVCodecsCore/Source/AVCodecsCore/AVCodecsCore.Build.cs
--- a/Engine/Plugins/Experimental/AVCodecs/AVCodecsCore/Source/AVCodecsCore/AVCodecsCore.Build.cs
+++ b/Engine/Plugins/Experimental/AVCodecs/AVCodecsCore/Source/AVCodecsCore/AVCodecsCore.Build.cs
@@ -11,22 +11,17 @@ public class AVCodecsCore : ModuleRules
     {
         // Without these two compilation fails on VS2017 with D8049: command line is too long to fit in debug record.
         bLegacyPublicIncludePaths = false;
-        DefaultBuildSettings = BuildSettingsVersion.V2;
-
-        // PCHUsage = PCHUsageMode.NoPCHs;
-
-        // PrecompileForTargets = PrecompileTargetsType.None;
+        DefaultBuildSettings = BuildSettingsVersion.Latest;
 
         PublicIncludePaths.AddRange(new string[] {
 			// ... add public include paths required here ...
 		});
 
         PrivateIncludePaths.AddRange(new string[] {
-			// ... add other private include paths required here ...
-		});
+        });
 
         PrivateDependencyModuleNames.AddRange(new string[] {
-            "Engine"
+            "CoreUObject",
         });
 
         PublicDependencyModuleNames.AddRange(new string[] {
@@ -43,12 +38,12 @@ public class AVCodecsCore : ModuleRules
             PublicIncludePathModuleNames.Add("Vulkan");
             AddEngineThirdPartyPrivateStaticDependencies(Target, "Vulkan");
 
-        	PublicDefinitions.Add("AVCODECS_USE_VULKAN=1");
-		} 
-		else
-		{
-			PublicDefinitions.Add("AVCODECS_USE_VULKAN=0");
-		}
+            PublicDefinitions.Add("AVCODECS_USE_VULKAN=1");
+        }
+        else
+        {
+            PublicDefinitions.Add("AVCODECS_USE_VULKAN=0");
+        }
 
         if (Target.IsInPlatformGroup(UnrealPlatformGroup.Windows))
         {
@@ -58,33 +53,33 @@ public class AVCodecsCore : ModuleRules
             PublicDelayLoadDLLs.Add("mfuuid.dll");
             PublicDelayLoadDLLs.Add("Mfreadwrite.dll");
 
-        	PublicDefinitions.Add("AVCODECS_USE_D3D=1");
-		}
-		else
-		{
-			PublicDefinitions.Add("AVCODECS_USE_D3D=0");
-		}
+            PublicDefinitions.Add("AVCODECS_USE_D3D=1");
+        }
+        else
+        {
+            PublicDefinitions.Add("AVCODECS_USE_D3D=0");
+        }
 
         if (Target.IsInPlatformGroup(UnrealPlatformGroup.Apple))
-		{
-			PublicFrameworks.AddRange(new string[]{
-				"AVFoundation",
-				"VideoToolbox"
-			});
-   
+        {
+            PublicFrameworks.AddRange(new string[]{
+                "AVFoundation",
+                "VideoToolbox"
+            });
+
             PublicDependencyModuleNames.AddRange(new string[] {
                 "MetalRHI",
             });
-   
+
             PublicWeakFrameworks.Add("Metal");
-			AddEngineThirdPartyPrivateStaticDependencies(Target, "MetalCPP");
-			
-			PublicDefinitions.Add("AVCODECS_USE_METAL=1");
-		}
-		else
-		{
-			PublicDefinitions.Add("AVCODECS_USE_METAL=0");
-		}
+            AddEngineThirdPartyPrivateStaticDependencies(Target, "MetalCPP");
+
+            PublicDefinitions.Add("AVCODECS_USE_METAL=1");
+        }
+        else
+        {
+            PublicDefinitions.Add("AVCODECS_USE_METAL=0");
+        }
 
         PublicDefinitions.Add("DEBUG_DUMP_TO_DISK=0");
     }
diff --git a/Engine/Plugins/Experimental/AVCodecs/AVCodecsCore/Source/AVCodecsCore/Public/AVConfig.h b/Engine/Plugins/Experimental/AVCodecs/AVCodecsCore/Source/AVCodecsCore/Public/AVConfig.h
--- a/Engine/Plugins/Experimental/AVCodecs/AVCodecsCore/Source/AVCodecsCore/Public/AVConfig.h
+++ b/Engine/Plugins/Experimental/AVCodecs/AVCodecsCore/Source/AVCodecsCore/Public/AVConfig.h
@@ -4,9 +4,12 @@
 
 #include "AVInstance.h"
 
+#include "AVConfig.generated.h"
+
 /**
  * Simple presets for whole AV pipeline.
  */
+UENUM()
 enum class EAVPreset : uint8
 {
 	UltraLowQuality,
@@ -19,6 +22,7 @@ enum class EAVPreset : uint8
 /**
  * Latency mode for the AV pipeline.
  */
+UENUM()
 enum class EAVLatencyMode : uint8
 {
 	UltraLowLatency,
diff --git a/Engine/Plugins/Experimental/AVCodecs/AVCodecsCore/Source/AVCodecsCore/Public/Video/CodecUtils/CodecUtilsH264.h b/Engine/Plugins/Experimental/AVCodecs/AVCodecsCore/Source/AVCodecsCore/Public/Video/CodecUtils/CodecUtilsH264.h
--- a/Engine/Plugins/Experimental/AVCodecs/AVCodecsCore/Source/AVCodecsCore/Public/Video/CodecUtils/CodecUtilsH264.h
+++ b/Engine/Plugins/Experimental/AVCodecs/AVCodecsCore/Source/AVCodecsCore/Public/Video/CodecUtils/CodecUtilsH264.h
@@ -7,12 +7,16 @@
 
 #include "AVResult.h"
 
+#include "CodecUtilsH264.generated.h"
+
 struct FH264ProfileDefinition;
 struct FVideoPacket;
 
 // Easy to understand user facing profiles from the T-Rec-H.264-201304 spec
 // there is another enum in CodecUtilsH264 that aligns the profiles with PIC and constraints
 // however users are not expected to understand those ones.
+
+UENUM()
 enum class EH264Profile : uint8
 {
 	Auto,
@@ -235,7 +239,7 @@ struct FH264ProfileDefinition
 	EH264Profile							   Profile;
 	UE::AVCodecCore::H264::EH264ProfileIDC	   PIDC;
 	UE::AVCodecCore::H264::EH264ConstraintFlag ConstraintFlags;
-	const TCHAR*							   Name;
+	const TCHAR* Name;
 };
 
 namespace UE::AVCodecCore::H264
diff --git a/Engine/Plugins/Experimental/AVCodecs/AVCodecsCore/Source/AVCodecsCore/Public/Video/VideoConfig.h b/Engine/Plugins/Experimental/AVCodecs/AVCodecsCore/Source/AVCodecsCore/Public/Video/VideoConfig.h
--- a/Engine/Plugins/Experimental/AVCodecs/AVCodecsCore/Source/AVCodecsCore/Public/Video/VideoConfig.h
+++ b/Engine/Plugins/Experimental/AVCodecs/AVCodecsCore/Source/AVCodecsCore/Public/Video/VideoConfig.h
@@ -4,6 +4,9 @@
 
 #include "AVCoder.h"
 
+#include "VideoConfig.generated.h"
+
+UENUM()
 enum class EVideoCodec : uint8
 {
 	Undefined,
diff --git a/Engine/Plugins/Experimental/AVCodecs/AVCodecsCore/Source/AVCodecsCore/Public/Video/VideoEncoder.h b/Engine/Plugins/Experimental/AVCodecs/AVCodecsCore/Source/AVCodecsCore/Public/Video/VideoEncoder.h
--- a/Engine/Plugins/Experimental/AVCodecs/AVCodecsCore/Source/AVCodecsCore/Public/Video/VideoEncoder.h
+++ b/Engine/Plugins/Experimental/AVCodecs/AVCodecsCore/Source/AVCodecsCore/Public/Video/VideoEncoder.h
@@ -8,9 +8,12 @@
 #include "Video/VideoPacket.h"
 #include "Video/VideoResource.h"
 
+#include "VideoEncoder.generated.h"
+
 /*
  * Implementation of Video Encoding domain, see TAVCoder for inheritance model
  */
+UENUM()
 enum class EScalabilityMode : uint8
 {
 	L1T1 = 0,
@@ -98,6 +101,7 @@ inline TOptional<EScalabilityMode> ScalabilityModeFromString(const FString& Mode
 	return FNullOpt(0);
 }
 
+UENUM()
 enum class ERateControlMode : uint8
 {
 	Unknown,
@@ -106,6 +110,7 @@ enum class ERateControlMode : uint8
 	CBR
 };
 
+UENUM()
 enum class EMultipassMode : uint8
 {
 	Unknown,
@@ -156,8 +161,8 @@ public:
 	// Advanced bitrate settings. Used for situations such as simulcast / SVC
 	TOptional<int32> Bitrates[Video::MaxSpatialLayers][Video::MaxTemporalStreams];
 
-	int32 MinQP = -1;
-	int32 MaxQP = -1;
+	int32 MinQuality = -1;
+	int32 MaxQuality = -1;
 
 	ERateControlMode RateControlMode = ERateControlMode::CBR;
 	uint8			 bFillData : 1;
@@ -181,56 +186,56 @@ public:
 	{
 		switch (Preset)
 		{
-			case EAVPreset::UltraLowQuality:
-				TargetBitrate = 500000;
-				MaxBitrate = 1000000;
+		case EAVPreset::UltraLowQuality:
+			TargetBitrate = 500000;
+			MaxBitrate = 1000000;
 
-				MinQP = 10;
-				MaxQP = 20;
+			MinQuality = 0;
+			MaxQuality = 33;
 
-				RateControlMode = ERateControlMode::CBR;
+			RateControlMode = ERateControlMode::CBR;
 
-				break;
-			case EAVPreset::LowQuality:
-				TargetBitrate = 3000000;
-				MaxBitrate = 4500000;
+			break;
+		case EAVPreset::LowQuality:
+			TargetBitrate = 3000000;
+			MaxBitrate = 4500000;
 
-				MinQP = 20;
-				MaxQP = 30;
+			MinQuality = 0;
+			MaxQuality = 50;
 
-				RateControlMode = ERateControlMode::CBR;
+			RateControlMode = ERateControlMode::CBR;
 
-				break;
-			case EAVPreset::Default:
-				TargetBitrate = 5000000;
-				MaxBitrate = 12500000;
+			break;
+		case EAVPreset::Default:
+			TargetBitrate = 5000000;
+			MaxBitrate = 12500000;
 
-				MinQP = 25;
-				MaxQP = 40;
+			MinQuality = 25;
+			MaxQuality = 75;
 
-				RateControlMode = ERateControlMode::CBR;
+			RateControlMode = ERateControlMode::CBR;
 
-				break;
-			case EAVPreset::HighQuality:
-				TargetBitrate = 10000000;
-				MaxBitrate = 20000000;
+			break;
+		case EAVPreset::HighQuality:
+			TargetBitrate = 10000000;
+			MaxBitrate = 20000000;
 
-				MinQP = 35;
-				MaxQP = 50;
+			MinQuality = 50;
+			MaxQuality = 100;
 
-				RateControlMode = ERateControlMode::VBR;
+			RateControlMode = ERateControlMode::VBR;
 
-				break;
-			case EAVPreset::Lossless:
-				TargetBitrate = 0;
-				MaxBitrate = 0;
+			break;
+		case EAVPreset::Lossless:
+			TargetBitrate = 0;
+			MaxBitrate = 0;
 
-				MinQP = -1;
-				MaxQP = -1;
+			MinQuality = -1;
+			MaxQuality = -1;
 
-				RateControlMode = ERateControlMode::ConstQP;
+			RateControlMode = ERateControlMode::ConstQP;
 
-				break;
+			break;
 		}
 	}
 };
diff --git a/Engine/Plugins/Experimental/AVCodecs/LibVpxCodecs/Source/LibVpxCodecs/Private/Video/Encoders/Configs/VideoEncoderConfigLibVpx.cpp b/Engine/Plugins/Experimental/AVCodecs/LibVpxCodecs/Source/LibVpxCodecs/Private/Video/Encoders/Configs/VideoEncoderConfigLibVpx.cpp
--- a/Engine/Plugins/Experimental/AVCodecs/LibVpxCodecs/Source/LibVpxCodecs/Private/Video/Encoders/Configs/VideoEncoderConfigLibVpx.cpp
+++ b/Engine/Plugins/Experimental/AVCodecs/LibVpxCodecs/Source/LibVpxCodecs/Private/Video/Encoders/Configs/VideoEncoderConfigLibVpx.cpp
@@ -21,8 +21,8 @@ DLLEXPORT FAVResult FAVExtension::TransformConfig(FVideoEncoderConfigLibVpx& Out
 	OutConfig.MinBitrate = InConfig.MinBitrate > -1 ? InConfig.MinBitrate : DEFAULT_BITRATE_MIN;
 	OutConfig.TargetBitrate = InConfig.TargetBitrate > -1 ? InConfig.TargetBitrate : DEFAULT_BITRATE_TARGET;
 	OutConfig.MaxBitrate = InConfig.MaxBitrate > -1 ? InConfig.MaxBitrate : DEFAULT_BITRATE_MAX;
-	OutConfig.MinQP = InConfig.MinQP;
-	OutConfig.MaxQP = InConfig.MaxQP;
+	OutConfig.MinQuality = InConfig.MinQuality;
+	OutConfig.MaxQuality = InConfig.MaxQuality;
 	OutConfig.KeyframeInterval = InConfig.KeyframeInterval > 0 ? InConfig.KeyframeInterval : 0;
 	OutConfig.ScalabilityMode = InConfig.ScalabilityMode;
 	OutConfig.NumberOfSpatialLayers = InConfig.NumberOfSpatialLayers;
@@ -55,8 +55,8 @@ DLLEXPORT FAVResult FAVExtension::TransformConfig(FVideoEncoderConfig& OutConfig
 	OutConfig.MinBitrate = InConfig.MinBitrate > -1 ? InConfig.MinBitrate : DEFAULT_BITRATE_MIN;
 	OutConfig.TargetBitrate = InConfig.TargetBitrate > -1 ? InConfig.TargetBitrate : DEFAULT_BITRATE_TARGET;
 	OutConfig.MaxBitrate = InConfig.MaxBitrate > -1 ? InConfig.MaxBitrate : DEFAULT_BITRATE_MAX;
-	OutConfig.MinQP = InConfig.MinQP;
-	OutConfig.MaxQP = InConfig.MaxQP;
+	OutConfig.MinQuality = InConfig.MinQuality;
+	OutConfig.MaxQuality = InConfig.MaxQuality;
 	OutConfig.KeyframeInterval = InConfig.KeyframeInterval > 0 ? InConfig.KeyframeInterval : 0;
 	OutConfig.ScalabilityMode = InConfig.ScalabilityMode;
 	OutConfig.NumberOfSpatialLayers = InConfig.NumberOfSpatialLayers;
diff --git a/Engine/Plugins/Experimental/AVCodecs/LibVpxCodecs/Source/LibVpxCodecs/Public/Video/Encoders/Configs/VideoEncoderConfigLibVpx.h b/Engine/Plugins/Experimental/AVCodecs/LibVpxCodecs/Source/LibVpxCodecs/Public/Video/Encoders/Configs/VideoEncoderConfigLibVpx.h
--- a/Engine/Plugins/Experimental/AVCodecs/LibVpxCodecs/Source/LibVpxCodecs/Public/Video/Encoders/Configs/VideoEncoderConfigLibVpx.h
+++ b/Engine/Plugins/Experimental/AVCodecs/LibVpxCodecs/Source/LibVpxCodecs/Public/Video/Encoders/Configs/VideoEncoderConfigLibVpx.h
@@ -22,8 +22,8 @@ public:
 	TOptional<int32> Bitrates[Video::MaxSpatialLayers][Video::MaxTemporalStreams];
 	uint32			 KeyframeInterval = 0;
 	EVideoFormat	 PixelFormat;
-	int32			 MinQP;
-	int32			 MaxQP;
+	int32			 MinQuality;
+	int32			 MaxQuality;
 
 	int32					 NumberOfCores = 0;
 	bool					 bDenoisingOn = false;
@@ -68,8 +68,8 @@ public:
 			&& this->MinBitrate == Other.MinBitrate
 			&& this->KeyframeInterval == Other.KeyframeInterval
 			&& this->PixelFormat == Other.PixelFormat
-			&& this->MinQP == Other.MinQP
-			&& this->MaxQP == Other.MaxQP
+			&& this->MinQuality == Other.MinQuality
+			&& this->MaxQuality == Other.MaxQuality
 			&& this->NumberOfCores == Other.NumberOfCores
 			&& this->bDenoisingOn == Other.bDenoisingOn
 			&& this->bAdaptiveQpMode == Other.bAdaptiveQpMode
diff --git a/Engine/Plugins/Experimental/AVCodecs/LibVpxCodecs/Source/LibVpxCodecs/Public/Video/Encoders/VideoEncoderLibVpxVP8.hpp b/Engine/Plugins/Experimental/AVCodecs/LibVpxCodecs/Source/LibVpxCodecs/Public/Video/Encoders/VideoEncoderLibVpxVP8.hpp
--- a/Engine/Plugins/Experimental/AVCodecs/LibVpxCodecs/Source/LibVpxCodecs/Public/Video/Encoders/VideoEncoderLibVpxVP8.hpp
+++ b/Engine/Plugins/Experimental/AVCodecs/LibVpxCodecs/Source/LibVpxCodecs/Public/Video/Encoders/VideoEncoderLibVpxVP8.hpp
@@ -135,13 +135,8 @@ FAVResult TVideoEncoderLibVpxVP8<TResource>::ApplyConfig()
 				VpxConfig->g_pass = VPX_RC_ONE_PASS;
 				// Handle resizing outside of libvpx.
 				VpxConfig->rc_resize_allowed = 0;
-				VpxConfig->rc_min_quantizer = 2;
-				if (PendingConfig.MaxQP >= static_cast<int32>(VpxConfig->rc_min_quantizer))
-				{
-					QpMax = PendingConfig.MaxQP;
-				}
-
-				VpxConfig->rc_max_quantizer = QpMax;
+				VpxConfig->rc_min_quantizer = (1.0f - (PendingConfig.MaxQuality / 100.0f)) * 63.0f;
+				VpxConfig->rc_max_quantizer = (1.0f - (PendingConfig.MinQuality / 100.0f)) * 63.0f;
 				VpxConfig->rc_undershoot_pct = 100;
 				VpxConfig->rc_overshoot_pct = 15;
 				VpxConfig->rc_buf_initial_sz = 500;
@@ -450,14 +445,14 @@ typename TVideoEncoderLibVpxVP8<TResource>::EEncodeResult TVideoEncoderLibVpxVP8
 	{
 		switch (Packet->kind)
 		{
-			case VPX_CODEC_CX_FRAME_PKT:
-			{
-				FMemory::Memcpy(&Buffer.GetData()[EncodedPos], Packet->data.frame.buf, Packet->data.frame.sz);
-				EncodedPos += Packet->data.frame.sz;
-				break;
-			}
-			default:
-				break;
+		case VPX_CODEC_CX_FRAME_PKT:
+		{
+			FMemory::Memcpy(&Buffer.GetData()[EncodedPos], Packet->data.frame.buf, Packet->data.frame.sz);
+			EncodedPos += Packet->data.frame.sz;
+			break;
+		}
+		default:
+			break;
 		}
 		// End of frame
 
diff --git a/Engine/Plugins/Experimental/AVCodecs/LibVpxCodecs/Source/LibVpxCodecs/Public/Video/Encoders/VideoEncoderLibVpxVP9.hpp b/Engine/Plugins/Experimental/AVCodecs/LibVpxCodecs/Source/LibVpxCodecs/Public/Video/Encoders/VideoEncoderLibVpxVP9.hpp
--- a/Engine/Plugins/Experimental/AVCodecs/LibVpxCodecs/Source/LibVpxCodecs/Public/Video/Encoders/VideoEncoderLibVpxVP9.hpp
+++ b/Engine/Plugins/Experimental/AVCodecs/LibVpxCodecs/Source/LibVpxCodecs/Public/Video/Encoders/VideoEncoderLibVpxVP9.hpp
@@ -102,8 +102,8 @@ FAVResult TVideoEncoderLibVpxVP9<TResource>::ApplyConfig()
 					&& this->AppliedConfig.Height == PendingConfig.Height
 					&& this->AppliedConfig.KeyframeInterval == PendingConfig.KeyframeInterval
 					&& this->AppliedConfig.PixelFormat == PendingConfig.PixelFormat
-					&& this->AppliedConfig.MinQP == PendingConfig.MinQP
-					&& this->AppliedConfig.MaxQP == PendingConfig.MaxQP
+					&& this->AppliedConfig.MinQuality == PendingConfig.MinQuality
+					&& this->AppliedConfig.MaxQuality == PendingConfig.MaxQuality
 					&& this->AppliedConfig.NumberOfCores == PendingConfig.NumberOfCores
 					&& this->AppliedConfig.bDenoisingOn == PendingConfig.bDenoisingOn
 					&& this->AppliedConfig.bAdaptiveQpMode == PendingConfig.bAdaptiveQpMode
@@ -206,29 +206,29 @@ FAVResult TVideoEncoderLibVpxVP9<TResource>::ApplyConfig()
 				unsigned int  BitsForStorage = 8;
 				switch (Profile)
 				{
-					case EProfile::Profile0:
-						PixelFormat = PreviousImgFmt.Get(VPX_IMG_FMT_I420);
-						BitsForStorage = 8;
-						VpxConfig->g_bit_depth = VPX_BITS_8;
-						VpxConfig->g_profile = 0;
-						VpxConfig->g_input_bit_depth = 8;
-						break;
-					case EProfile::Profile1:
-						// Encoding of profile 1 is not implemented. It would require extended
-						// support for I444, I422, and I440 buffers.
-						checkNoEntry();
-						break;
-					case EProfile::Profile2:
-						PixelFormat = VPX_IMG_FMT_I42016;
-						BitsForStorage = 16;
-						VpxConfig->g_bit_depth = VPX_BITS_10;
-						VpxConfig->g_profile = 2;
-						VpxConfig->g_input_bit_depth = 10;
-						break;
-					case EProfile::Profile3:
-						// Encoding of profile 3 is not implemented.
-						checkNoEntry();
-						break;
+				case EProfile::Profile0:
+					PixelFormat = PreviousImgFmt.Get(VPX_IMG_FMT_I420);
+					BitsForStorage = 8;
+					VpxConfig->g_bit_depth = VPX_BITS_8;
+					VpxConfig->g_profile = 0;
+					VpxConfig->g_input_bit_depth = 8;
+					break;
+				case EProfile::Profile1:
+					// Encoding of profile 1 is not implemented. It would require extended
+					// support for I444, I422, and I440 buffers.
+					checkNoEntry();
+					break;
+				case EProfile::Profile2:
+					PixelFormat = VPX_IMG_FMT_I42016;
+					BitsForStorage = 16;
+					VpxConfig->g_bit_depth = VPX_BITS_10;
+					VpxConfig->g_profile = 2;
+					VpxConfig->g_input_bit_depth = 10;
+					break;
+				case EProfile::Profile3:
+					// Encoding of profile 3 is not implemented.
+					checkNoEntry();
+					break;
 				}
 
 				// Creating a wrapper to the image - setting image data to nullptr. Actual
@@ -254,8 +254,8 @@ FAVResult TVideoEncoderLibVpxVP9<TResource>::ApplyConfig()
 				VpxConfig->rc_dropframe_thresh = 0;
 				VpxConfig->rc_end_usage = VPX_CBR;
 				VpxConfig->g_pass = VPX_RC_ONE_PASS;
-				VpxConfig->rc_min_quantizer = 2;
-				VpxConfig->rc_max_quantizer = 52;
+				VpxConfig->rc_min_quantizer = (1.0f - (PendingConfig.MaxQuality / 100.0f)) * 63.0f;
+				VpxConfig->rc_max_quantizer = (1.0f - (PendingConfig.MinQuality / 100.0f)) * 63.0f;
 				VpxConfig->rc_undershoot_pct = 50;
 				VpxConfig->rc_overshoot_pct = 50;
 				VpxConfig->rc_buf_initial_sz = 500;
@@ -493,44 +493,44 @@ FAVResult TVideoEncoderLibVpxVP9<TResource>::SendFrame(TSharedPtr<FVideoResource
 
 		switch (Profile)
 		{
-			case EProfile::Profile0:
-			{
-				int StrideY = Resource->GetWidth();
-				int StrideUV = (Resource->GetWidth() + 1) / 2;
+		case EProfile::Profile0:
+		{
+			int StrideY = Resource->GetWidth();
+			int StrideUV = (Resource->GetWidth() + 1) / 2;
 
-				int DataSizeY = StrideY * Resource->GetHeight();
-				int DataSizeUV = StrideUV * ((Resource->GetHeight() + 1) / 2);
+			int DataSizeY = StrideY * Resource->GetHeight();
+			int DataSizeUV = StrideUV * ((Resource->GetHeight() + 1) / 2);
 
-				RawImage->planes[VPX_PLANE_Y] = static_cast<uint8*>(Resource->GetRaw().Get());
-				RawImage->planes[VPX_PLANE_U] = static_cast<uint8*>(Resource->GetRaw().Get()) + DataSizeY;
-				RawImage->planes[VPX_PLANE_V] = static_cast<uint8*>(Resource->GetRaw().Get()) + DataSizeY + DataSizeUV;
-				RawImage->stride[VPX_PLANE_Y] = StrideY;
-				RawImage->stride[VPX_PLANE_U] = StrideUV;
-				RawImage->stride[VPX_PLANE_V] = StrideUV;
+			RawImage->planes[VPX_PLANE_Y] = static_cast<uint8*>(Resource->GetRaw().Get());
+			RawImage->planes[VPX_PLANE_U] = static_cast<uint8*>(Resource->GetRaw().Get()) + DataSizeY;
+			RawImage->planes[VPX_PLANE_V] = static_cast<uint8*>(Resource->GetRaw().Get()) + DataSizeY + DataSizeUV;
+			RawImage->stride[VPX_PLANE_Y] = StrideY;
+			RawImage->stride[VPX_PLANE_U] = StrideUV;
+			RawImage->stride[VPX_PLANE_V] = StrideUV;
 
-				break;
-			}
-			case EProfile::Profile1:
-				checkNoEntry();
-				break;
-			case EProfile::Profile2:
-			{
-				int StrideY = Resource->GetWidth();
-				int StrideUV = (Resource->GetWidth() + 1) / 2;
-
-				int DataSizeY = 2 * StrideY * Resource->GetHeight();
-				int DataSizeUV = 2 * StrideUV * ((Resource->GetHeight() + 1) / 2);
-
-				RawImage->planes[VPX_PLANE_Y] = static_cast<uint8*>(Resource->GetRaw().Get());
-				RawImage->planes[VPX_PLANE_U] = static_cast<uint8*>(Resource->GetRaw().Get()) + DataSizeY;
-				RawImage->planes[VPX_PLANE_V] = static_cast<uint8*>(Resource->GetRaw().Get()) + DataSizeY + DataSizeUV;
-				RawImage->stride[VPX_PLANE_Y] = StrideY * 2;
-				RawImage->stride[VPX_PLANE_U] = StrideUV * 2;
-				RawImage->stride[VPX_PLANE_V] = StrideUV * 2;
-			}
-			case EProfile::Profile3:
-				checkNoEntry();
-				break;
+			break;
+		}
+		case EProfile::Profile1:
+			checkNoEntry();
+			break;
+		case EProfile::Profile2:
+		{
+			int StrideY = Resource->GetWidth();
+			int StrideUV = (Resource->GetWidth() + 1) / 2;
+
+			int DataSizeY = 2 * StrideY * Resource->GetHeight();
+			int DataSizeUV = 2 * StrideUV * ((Resource->GetHeight() + 1) / 2);
+
+			RawImage->planes[VPX_PLANE_Y] = static_cast<uint8*>(Resource->GetRaw().Get());
+			RawImage->planes[VPX_PLANE_U] = static_cast<uint8*>(Resource->GetRaw().Get()) + DataSizeY;
+			RawImage->planes[VPX_PLANE_V] = static_cast<uint8*>(Resource->GetRaw().Get()) + DataSizeY + DataSizeUV;
+			RawImage->stride[VPX_PLANE_Y] = StrideY * 2;
+			RawImage->stride[VPX_PLANE_U] = StrideUV * 2;
+			RawImage->stride[VPX_PLANE_V] = StrideUV * 2;
+		}
+		case EProfile::Profile3:
+			checkNoEntry();
+			break;
 		}
 
 		vpx_enc_frame_flags_t Flags = 0;
@@ -788,17 +788,17 @@ FAVResult TVideoEncoderLibVpxVP9<TResource>::InitAndSetControlSettings(FVideoEnc
 	{
 		switch (InterLayerPrediction)
 		{
-			case EInterLayerPrediction::On:
-				SAFECONTROLVP9(Encoder.Get(), VP9E_SET_SVC_INTER_LAYER_PRED, 0);
-				break;
-			case EInterLayerPrediction::Off:
-				SAFECONTROLVP9(Encoder.Get(), VP9E_SET_SVC_INTER_LAYER_PRED, 1);
-				break;
-			case EInterLayerPrediction::OnKeyPicture:
-				SAFECONTROLVP9(Encoder.Get(), VP9E_SET_SVC_INTER_LAYER_PRED, 2);
-				break;
-			default:
-				checkNoEntry();
+		case EInterLayerPrediction::On:
+			SAFECONTROLVP9(Encoder.Get(), VP9E_SET_SVC_INTER_LAYER_PRED, 0);
+			break;
+		case EInterLayerPrediction::Off:
+			SAFECONTROLVP9(Encoder.Get(), VP9E_SET_SVC_INTER_LAYER_PRED, 1);
+			break;
+		case EInterLayerPrediction::OnKeyPicture:
+			SAFECONTROLVP9(Encoder.Get(), VP9E_SET_SVC_INTER_LAYER_PRED, 2);
+			break;
+		default:
+			checkNoEntry();
 		}
 
 		if (!SvcDropFrame.IsValid())
@@ -1242,7 +1242,7 @@ bool TVideoEncoderLibVpxVP9<TResource>::PopulateCodecSpecific(FCodecSpecificInfo
 	{
 		FScalableVideoController::FLayerFrameConfig* FoundLayer = LayerFrames.FindByPredicate([&LayerId](FScalableVideoController::FLayerFrameConfig& Config) {
 			return Config.GetSpatialId() == LayerId.spatial_layer_id;
-		});
+			});
 		if (FoundLayer == nullptr)
 		{
 			// UE_LOGFMT(LogPixelStreamingEpicRtc, Error, "Encoder produced a frame for layer S{0}T{1} that wasn't requested", LayerId.spatial_layer_id, LayerId.temporal_layer_id);
@@ -1395,10 +1395,10 @@ void TVideoEncoderLibVpxVP9<TResource>::UpdateReferenceBuffers(const vpx_codec_c
 	vpx_svc_layer_id_t LayerId = { 0 };
 	SAFECONTROLVP9(Encoder.Get(), VP9E_GET_SVC_LAYER_ID, &LayerId)
 
-	FRefFrameBuffer FrameBuf = {
-		.PicNum = PicNum,
-		.SpatialLayerId = LayerId.spatial_layer_id,
-		.TemporalLayerId = LayerId.temporal_layer_id
+		FRefFrameBuffer FrameBuf = {
+			.PicNum = PicNum,
+			.SpatialLayerId = LayerId.spatial_layer_id,
+			.TemporalLayerId = LayerId.temporal_layer_id
 	};
 
 	if (bIsSvc)
@@ -1748,7 +1748,7 @@ void TVideoEncoderLibVpxVP9<TResource>::UpdatePerformanceFlags(FVideoEncoderConf
 		}
 
 		return ParamsByResolution[ResolutionArray[Index]];
-	};
+		};
 	PerformanceFlagsBySpatialIndex.Empty();
 
 	if (bIsSvc)
diff --git a/Engine/Plugins/Experimental/AVCodecs/NVCodecs/Source/NVENC/Private/Video/Encoders/Configs/VideoEncoderConfigNVENC.cpp b/Engine/Plugins/Experimental/AVCodecs/NVCodecs/Source/NVENC/Private/Video/Encoders/Configs/VideoEncoderConfigNVENC.cpp
--- a/Engine/Plugins/Experimental/AVCodecs/NVCodecs/Source/NVENC/Private/Video/Encoders/Configs/VideoEncoderConfigNVENC.cpp
+++ b/Engine/Plugins/Experimental/AVCodecs/NVCodecs/Source/NVENC/Private/Video/Encoders/Configs/VideoEncoderConfigNVENC.cpp
@@ -12,14 +12,14 @@ TAVResult<NV_ENC_PARAMS_RC_MODE> FVideoEncoderConfigNVENC::ConvertRateControlMod
 {
 	switch (Mode)
 	{
-		case ERateControlMode::ConstQP:
-			return NV_ENC_PARAMS_RC_CONSTQP;
-		case ERateControlMode::VBR:
-			return NV_ENC_PARAMS_RC_VBR;
-		case ERateControlMode::CBR:
-			return NV_ENC_PARAMS_RC_CBR;
-		default:
-			return FAVResult(EAVResult::ErrorUnsupported, FString::Printf(TEXT("Rate control mode %d is not supported"), Mode), TEXT("NVENC"));
+	case ERateControlMode::ConstQP:
+		return NV_ENC_PARAMS_RC_CONSTQP;
+	case ERateControlMode::VBR:
+		return NV_ENC_PARAMS_RC_VBR;
+	case ERateControlMode::CBR:
+		return NV_ENC_PARAMS_RC_CBR;
+	default:
+		return FAVResult(EAVResult::ErrorUnsupported, FString::Printf(TEXT("Rate control mode %d is not supported"), Mode), TEXT("NVENC"));
 	}
 }
 
@@ -27,14 +27,14 @@ TAVResult<ERateControlMode> FVideoEncoderConfigNVENC::ConvertRateControlMode(NV_
 {
 	switch (Mode)
 	{
-		case NV_ENC_PARAMS_RC_CONSTQP:
-			return ERateControlMode::ConstQP;
-		case NV_ENC_PARAMS_RC_VBR:
-			return ERateControlMode::VBR;
-		case NV_ENC_PARAMS_RC_CBR:
-			return ERateControlMode::CBR;
-		default:
-			return FAVResult(EAVResult::ErrorUnsupported, FString::Printf(TEXT("Rate control mode %d is not supported"), Mode), TEXT("NVENC"));
+	case NV_ENC_PARAMS_RC_CONSTQP:
+		return ERateControlMode::ConstQP;
+	case NV_ENC_PARAMS_RC_VBR:
+		return ERateControlMode::VBR;
+	case NV_ENC_PARAMS_RC_CBR:
+		return ERateControlMode::CBR;
+	default:
+		return FAVResult(EAVResult::ErrorUnsupported, FString::Printf(TEXT("Rate control mode %d is not supported"), Mode), TEXT("NVENC"));
 	}
 }
 
@@ -42,14 +42,14 @@ TAVResult<NV_ENC_MULTI_PASS> FVideoEncoderConfigNVENC::ConvertMultipassMode(EMul
 {
 	switch (Mode)
 	{
-		case EMultipassMode::Disabled:
-			return NV_ENC_MULTI_PASS_DISABLED;
-		case EMultipassMode::Quarter:
-			return NV_ENC_TWO_PASS_QUARTER_RESOLUTION;
-		case EMultipassMode::Full:
-			return NV_ENC_TWO_PASS_FULL_RESOLUTION;
-		default:
-			return FAVResult(EAVResult::ErrorUnsupported, FString::Printf(TEXT("Multipass mode %d is not supported"), Mode), TEXT("NVENC"));
+	case EMultipassMode::Disabled:
+		return NV_ENC_MULTI_PASS_DISABLED;
+	case EMultipassMode::Quarter:
+		return NV_ENC_TWO_PASS_QUARTER_RESOLUTION;
+	case EMultipassMode::Full:
+		return NV_ENC_TWO_PASS_FULL_RESOLUTION;
+	default:
+		return FAVResult(EAVResult::ErrorUnsupported, FString::Printf(TEXT("Multipass mode %d is not supported"), Mode), TEXT("NVENC"));
 	}
 }
 
@@ -57,14 +57,14 @@ TAVResult<EMultipassMode> FVideoEncoderConfigNVENC::ConvertMultipassMode(NV_ENC_
 {
 	switch (Mode)
 	{
-		case NV_ENC_MULTI_PASS_DISABLED:
-			return EMultipassMode::Disabled;
-		case NV_ENC_TWO_PASS_QUARTER_RESOLUTION:
-			return EMultipassMode::Quarter;
-		case NV_ENC_TWO_PASS_FULL_RESOLUTION:
-			return EMultipassMode::Full;
-		default:
-			return FAVResult(EAVResult::ErrorUnsupported, FString::Printf(TEXT("Multipass mode %d is not supported"), Mode), TEXT("NVENC"));
+	case NV_ENC_MULTI_PASS_DISABLED:
+		return EMultipassMode::Disabled;
+	case NV_ENC_TWO_PASS_QUARTER_RESOLUTION:
+		return EMultipassMode::Quarter;
+	case NV_ENC_TWO_PASS_FULL_RESOLUTION:
+		return EMultipassMode::Full;
+	default:
+		return FAVResult(EAVResult::ErrorUnsupported, FString::Printf(TEXT("Multipass mode %d is not supported"), Mode), TEXT("NVENC"));
 	}
 }
 
@@ -72,16 +72,16 @@ TAVResult<NV_ENC_BUFFER_FORMAT> FVideoEncoderConfigNVENC::ConvertFormat(EVideoFo
 {
 	switch (Format)
 	{
-		case EVideoFormat::BGRA:
-			return NV_ENC_BUFFER_FORMAT::NV_ENC_BUFFER_FORMAT_ARGB;
-		case EVideoFormat::ABGR10:
-			return NV_ENC_BUFFER_FORMAT::NV_ENC_BUFFER_FORMAT_ARGB10;
-		case EVideoFormat::NV12:
-			return NV_ENC_BUFFER_FORMAT::NV_ENC_BUFFER_FORMAT_NV12;
-		case EVideoFormat::P010:
-			return NV_ENC_BUFFER_FORMAT::NV_ENC_BUFFER_FORMAT_YUV420_10BIT;
-		default:
-			return FAVResult(EAVResult::ErrorUnsupported, FString::Printf(TEXT("Pixel format %d is not supported"), Format), TEXT("NVENC"));
+	case EVideoFormat::BGRA:
+		return NV_ENC_BUFFER_FORMAT::NV_ENC_BUFFER_FORMAT_ARGB;
+	case EVideoFormat::ABGR10:
+		return NV_ENC_BUFFER_FORMAT::NV_ENC_BUFFER_FORMAT_ARGB10;
+	case EVideoFormat::NV12:
+		return NV_ENC_BUFFER_FORMAT::NV_ENC_BUFFER_FORMAT_NV12;
+	case EVideoFormat::P010:
+		return NV_ENC_BUFFER_FORMAT::NV_ENC_BUFFER_FORMAT_YUV420_10BIT;
+	default:
+		return FAVResult(EAVResult::ErrorUnsupported, FString::Printf(TEXT("Pixel format %d is not supported"), Format), TEXT("NVENC"));
 	}
 
 	// NV_ENC_BUFFER_FORMAT::NV_ENC_BUFFER_FORMAT_YUV420_10BIT : NV_ENC_BUFFER_FORMAT::NV_ENC_BUFFER_FORMAT_IYUV;
@@ -104,23 +104,23 @@ DLLEXPORT FAVResult FAVExtension::TransformConfig(FVideoEncoderConfigNVENC& OutC
 
 	switch (InConfig.Preset)
 	{
-		case EAVPreset::Lossless:
-			OutConfig.tuningInfo = NV_ENC_TUNING_INFO_LOSSLESS;
-			break;
-		case EAVPreset::HighQuality:
-			OutConfig.tuningInfo = NV_ENC_TUNING_INFO_HIGH_QUALITY;
-			break;
-		case EAVPreset::LowQuality:
-			OutConfig.tuningInfo = NV_ENC_TUNING_INFO_LOW_LATENCY;
-			break;
-		default:
-		case EAVPreset::UltraLowQuality:
-			OutConfig.tuningInfo = NV_ENC_TUNING_INFO_ULTRA_LOW_LATENCY;
-			break;
+	case EAVPreset::Lossless:
+		OutConfig.tuningInfo = NV_ENC_TUNING_INFO_LOSSLESS;
+		break;
+	case EAVPreset::HighQuality:
+		OutConfig.tuningInfo = NV_ENC_TUNING_INFO_HIGH_QUALITY;
+		break;
+	case EAVPreset::LowQuality:
+		OutConfig.tuningInfo = NV_ENC_TUNING_INFO_LOW_LATENCY;
+		break;
+	default:
+	case EAVPreset::UltraLowQuality:
+		OutConfig.tuningInfo = NV_ENC_TUNING_INFO_ULTRA_LOW_LATENCY;
+		break;
 	}
 
-	uint32_t const MinQP = static_cast<uint32>(InConfig.MinQP);
-	uint32_t const MaxQP = static_cast<uint32>(InConfig.MaxQP);
+	uint32_t const MinQP = (1.0f - (InConfig.MaxQuality / 100.0f)) * 51.0f;
+	uint32_t const MaxQP = (1.0f - (InConfig.MinQuality / 100.0f)) * 51.0f;
 
 	ERateControlMode ActualRateControlMode = InConfig.Preset == EAVPreset::Lossless ? ERateControlMode::ConstQP : InConfig.RateControlMode;
 	TAVResult<NV_ENC_PARAMS_RC_MODE> const ConvertedRateControlMode = FVideoEncoderConfigNVENC::ConvertRateControlMode(ActualRateControlMode);
@@ -143,8 +143,8 @@ DLLEXPORT FAVResult FAVExtension::TransformConfig(FVideoEncoderConfigNVENC& OutC
 	RateControlParams.minQP = { MinQP, MinQP, MinQP };
 	RateControlParams.maxQP = { MaxQP, MaxQP, MaxQP };
 	RateControlParams.constQP = RateControlParams.maxQP;
-	RateControlParams.enableMinQP = InConfig.MinQP > -1;
-	RateControlParams.enableMaxQP = InConfig.MaxQP > -1;
+	RateControlParams.enableMinQP = InConfig.MaxQuality < 100;
+	RateControlParams.enableMaxQP = InConfig.MinQuality > 0;
 
 	// IDR period - how often to send IDR (instantaneous decode refresh) frames, a.k.a keyframes. This can stabilise a stream that dropped/lost some frames (but at the cost of more bandwidth)
 	if (!FMemory::Memcmp(&OutConfig.encodeGUID, &NV_ENC_CODEC_H264_GUID, sizeof(GUID)))
@@ -183,8 +183,8 @@ DLLEXPORT FAVResult FAVExtension::TransformConfig(FVideoEncoderConfig& OutConfig
 	OutConfig.Height = InConfig.encodeHeight;
 
 	NV_ENC_RC_PARAMS& RateControlParams = InConfig.encodeConfig->rcParams;
-	OutConfig.MinQP = static_cast<int32>(RateControlParams.minQP.qpIntra);
-	OutConfig.MaxQP = static_cast<int32>(RateControlParams.maxQP.qpIntra);
+	OutConfig.MinQuality = static_cast<int32>((1.0f - (RateControlParams.maxQP.qpIntra / 51.0f)) * 100.0f);
+	OutConfig.MaxQuality = static_cast<int32>((1.0f - (RateControlParams.minQP.qpIntra / 51.0f)) * 100.0f);
 	OutConfig.TargetBitrate = RateControlParams.averageBitRate == DEFAULT_BITRATE_TARGET ? -1 : RateControlParams.averageBitRate;
 	OutConfig.MaxBitrate = RateControlParams.maxBitRate == DEFAULT_BITRATE_MAX ? -1 : RateControlParams.maxBitRate;
 
@@ -228,28 +228,28 @@ DLLEXPORT FAVResult FAVExtension::TransformConfig(FVideoEncoderConfigNVENC& OutC
 	static auto const ConvertProfile = [](EH264Profile Profile) -> TAVResult<GUID> {
 		switch (Profile)
 		{
-			case EH264Profile::Auto:
-				return NV_ENC_CODEC_PROFILE_AUTOSELECT_GUID;
-			case EH264Profile::Baseline:
-				return NV_ENC_H264_PROFILE_BASELINE_GUID;
-			case EH264Profile::Main:
-				return NV_ENC_H264_PROFILE_MAIN_GUID;
-			case EH264Profile::High:
-				return NV_ENC_H264_PROFILE_HIGH_GUID;
-			case EH264Profile::ProgressiveHigh:
-				return NV_ENC_H264_PROFILE_PROGRESSIVE_HIGH_GUID;
-			case EH264Profile::ConstrainedHigh:
-				return NV_ENC_H264_PROFILE_CONSTRAINED_HIGH_GUID;
-			case EH264Profile::High10:
-				return NV_ENC_HEVC_PROFILE_MAIN10_GUID;
-			case EH264Profile::StereoHigh:
-				return NV_ENC_H264_PROFILE_STEREO_GUID;
-			case EH264Profile::High444:
-				return NV_ENC_H264_PROFILE_HIGH_444_GUID;
-			default:
-				return FAVResult(EAVResult::ErrorUnsupported, FString::Printf(TEXT("H264 profile %d is not supported"), Profile), TEXT("NVENC"));
+		case EH264Profile::Auto:
+			return NV_ENC_CODEC_PROFILE_AUTOSELECT_GUID;
+		case EH264Profile::Baseline:
+			return NV_ENC_H264_PROFILE_BASELINE_GUID;
+		case EH264Profile::Main:
+			return NV_ENC_H264_PROFILE_MAIN_GUID;
+		case EH264Profile::High:
+			return NV_ENC_H264_PROFILE_HIGH_GUID;
+		case EH264Profile::ProgressiveHigh:
+			return NV_ENC_H264_PROFILE_PROGRESSIVE_HIGH_GUID;
+		case EH264Profile::ConstrainedHigh:
+			return NV_ENC_H264_PROFILE_CONSTRAINED_HIGH_GUID;
+		case EH264Profile::High10:
+			return NV_ENC_HEVC_PROFILE_MAIN10_GUID;
+		case EH264Profile::StereoHigh:
+			return NV_ENC_H264_PROFILE_STEREO_GUID;
+		case EH264Profile::High444:
+			return NV_ENC_H264_PROFILE_HIGH_444_GUID;
+		default:
+			return FAVResult(EAVResult::ErrorUnsupported, FString::Printf(TEXT("H264 profile %d is not supported"), Profile), TEXT("NVENC"));
 		}
-	};
+		};
 
 	OutConfig.encodeGUID = NV_ENC_CODEC_H264_GUID;
 
@@ -269,16 +269,16 @@ DLLEXPORT FAVResult FAVExtension::TransformConfig(FVideoEncoderConfigNVENC& OutC
 	static auto const ConvertAdaptiveTransformMode = [](EH264AdaptiveTransformMode TransformMode) -> TAVResult<NV_ENC_H264_ADAPTIVE_TRANSFORM_MODE> {
 		switch (TransformMode)
 		{
-			case EH264AdaptiveTransformMode::Auto:
-				return NV_ENC_H264_ADAPTIVE_TRANSFORM_AUTOSELECT;
-			case EH264AdaptiveTransformMode::Disable:
-				return NV_ENC_H264_ADAPTIVE_TRANSFORM_DISABLE;
-			case EH264AdaptiveTransformMode::Enable:
-				return NV_ENC_H264_ADAPTIVE_TRANSFORM_ENABLE;
-			default:
-				return FAVResult(EAVResult::ErrorUnsupported, FString::Printf(TEXT("H264 transform mode %d is not supported"), TransformMode), TEXT("NVENC"));
+		case EH264AdaptiveTransformMode::Auto:
+			return NV_ENC_H264_ADAPTIVE_TRANSFORM_AUTOSELECT;
+		case EH264AdaptiveTransformMode::Disable:
+			return NV_ENC_H264_ADAPTIVE_TRANSFORM_DISABLE;
+		case EH264AdaptiveTransformMode::Enable:
+			return NV_ENC_H264_ADAPTIVE_TRANSFORM_ENABLE;
+		default:
+			return FAVResult(EAVResult::ErrorUnsupported, FString::Printf(TEXT("H264 transform mode %d is not supported"), TransformMode), TEXT("NVENC"));
 		}
-	};
+		};
 
 	TAVResult<NV_ENC_H264_ADAPTIVE_TRANSFORM_MODE> const ConvertedAdaptiveTransformMode = ConvertAdaptiveTransformMode(InConfig.AdaptiveTransformMode);
 	if (ConvertedAdaptiveTransformMode.IsNotSuccess())
@@ -291,16 +291,16 @@ DLLEXPORT FAVResult FAVExtension::TransformConfig(FVideoEncoderConfigNVENC& OutC
 	static auto const ConvertEntropyCodingMode = [](EH264EntropyCodingMode EntropyCodingMode) -> TAVResult<NV_ENC_H264_ENTROPY_CODING_MODE> {
 		switch (EntropyCodingMode)
 		{
-			case EH264EntropyCodingMode::Auto:
-				return NV_ENC_H264_ENTROPY_CODING_MODE_AUTOSELECT;
-			case EH264EntropyCodingMode::CABAC:
-				return NV_ENC_H264_ENTROPY_CODING_MODE_CABAC;
-			case EH264EntropyCodingMode::CAVLC:
-				return NV_ENC_H264_ENTROPY_CODING_MODE_CAVLC;
-			default:
-				return FAVResult(EAVResult::ErrorUnsupported, FString::Printf(TEXT("H264 entropy coding mode %d is not supported"), EntropyCodingMode), TEXT("NVENC"));
+		case EH264EntropyCodingMode::Auto:
+			return NV_ENC_H264_ENTROPY_CODING_MODE_AUTOSELECT;
+		case EH264EntropyCodingMode::CABAC:
+			return NV_ENC_H264_ENTROPY_CODING_MODE_CABAC;
+		case EH264EntropyCodingMode::CAVLC:
+			return NV_ENC_H264_ENTROPY_CODING_MODE_CAVLC;
+		default:
+			return FAVResult(EAVResult::ErrorUnsupported, FString::Printf(TEXT("H264 entropy coding mode %d is not supported"), EntropyCodingMode), TEXT("NVENC"));
 		}
-	};
+		};
 
 	TAVResult<NV_ENC_H264_ENTROPY_CODING_MODE> const ConvertedEntropyCodingMode = ConvertEntropyCodingMode(InConfig.EntropyCodingMode);
 	if (ConvertedEntropyCodingMode.IsNotSuccess())
@@ -352,17 +352,17 @@ DLLEXPORT FAVResult FAVExtension::TransformConfig(FVideoEncoderConfigNVENC& OutC
 	static auto const ConvertProfile = [](EH265Profile Profile) -> TAVResult<GUID> {
 		switch (Profile)
 		{
-			case EH265Profile::Auto:
-			case EH265Profile::Main:
-				return NV_ENC_HEVC_PROFILE_MAIN_GUID;
-			case EH265Profile::Main10:
-				return NV_ENC_HEVC_PROFILE_MAIN10_GUID;
-			case EH265Profile::Main12:
-				return NV_ENC_HEVC_PROFILE_FREXT_GUID;
-			default:
-				return FAVResult(EAVResult::ErrorUnsupported, FString::Printf(TEXT("H265 profile %d is not supported"), Profile), TEXT("NVENC"));
+		case EH265Profile::Auto:
+		case EH265Profile::Main:
+			return NV_ENC_HEVC_PROFILE_MAIN_GUID;
+		case EH265Profile::Main10:
+			return NV_ENC_HEVC_PROFILE_MAIN10_GUID;
+		case EH265Profile::Main12:
+			return NV_ENC_HEVC_PROFILE_FREXT_GUID;
+		default:
+			return FAVResult(EAVResult::ErrorUnsupported, FString::Printf(TEXT("H265 profile %d is not supported"), Profile), TEXT("NVENC"));
 		}
-	};
+		};
 
 	OutConfig.encodeGUID = NV_ENC_CODEC_HEVC_GUID;
 
@@ -395,16 +395,16 @@ DLLEXPORT FAVResult FAVExtension::TransformConfig(FVideoEncoderConfigNVENC& OutC
 
 	switch (InConfig.Profile)
 	{
-		case EH265Profile::Main:
-		default:
-			OutH265Config.pixelBitDepthMinus8 = 0;
-			break;
-		case EH265Profile::Main10:
-			OutH265Config.pixelBitDepthMinus8 = 2;
-			break;
-		case EH265Profile::Main12:
-			OutH265Config.pixelBitDepthMinus8 = 4;
-			break;
+	case EH265Profile::Main:
+	default:
+		OutH265Config.pixelBitDepthMinus8 = 0;
+		break;
+	case EH265Profile::Main10:
+		OutH265Config.pixelBitDepthMinus8 = 2;
+		break;
+	case EH265Profile::Main12:
+		OutH265Config.pixelBitDepthMinus8 = 4;
+		break;
 	}
 
 	// Intra refresh - used to stabilise stream on the decoded side when frames are dropped/lost.
@@ -424,13 +424,13 @@ DLLEXPORT FAVResult FAVExtension::TransformConfig(FVideoEncoderConfigNVENC& OutC
 	static auto const ConvertProfile = [](EAV1Profile Profile) -> TAVResult<GUID> {
 		switch (Profile)
 		{
-			case EAV1Profile::Auto:
-			case EAV1Profile::Main:
-				return NV_ENC_AV1_PROFILE_MAIN_GUID;
-			default:
-				return FAVResult(EAVResult::ErrorUnsupported, FString::Printf(TEXT("AV1 profile %d is not supported"), Profile), TEXT("NVENC"));
+		case EAV1Profile::Auto:
+		case EAV1Profile::Main:
+			return NV_ENC_AV1_PROFILE_MAIN_GUID;
+		default:
+			return FAVResult(EAVResult::ErrorUnsupported, FString::Printf(TEXT("AV1 profile %d is not supported"), Profile), TEXT("NVENC"));
 		}
-	};
+		};
 
 	OutConfig.encodeGUID = NV_ENC_CODEC_AV1_GUID;
 
@@ -453,7 +453,7 @@ DLLEXPORT FAVResult FAVExtension::TransformConfig(FVideoEncoderConfigNVENC& OutC
 
 	// Repeat sequence header - sends sequence header with every IDR frame - maximum stabilisation of the stream when IDR is sent.
 	OutAV1Config.repeatSeqHdr = InConfig.RepeatSeqHdr;
-	
+
 	OutAV1Config.level = NV_ENC_LEVEL_AV1_AUTOSELECT;
 	OutAV1Config.colorPrimaries = NV_ENC_VUI_COLOR_PRIMARIES_UNSPECIFIED;
 	OutAV1Config.transferCharacteristics = NV_ENC_VUI_TRANSFER_CHARACTERISTIC_UNSPECIFIED;
@@ -461,11 +461,11 @@ DLLEXPORT FAVResult FAVExtension::TransformConfig(FVideoEncoderConfigNVENC& OutC
 
 	switch (InConfig.Profile)
 	{
-		case EAV1Profile::Auto:
-		case EAV1Profile::Main:
-		default:
-			OutAV1Config.pixelBitDepthMinus8 = 0;
-			break;
+	case EAV1Profile::Auto:
+	case EAV1Profile::Main:
+	default:
+		OutAV1Config.pixelBitDepthMinus8 = 0;
+		break;
 	}
 
 	// Intra refresh - used to stabilise stream on the decoded side when frames are dropped/lost.
diff --git a/Engine/Plugins/Media/PixelStreaming/Source/PixelStreaming/Private/VideoEncoderSingleLayerHardware.cpp b/Engine/Plugins/Media/PixelStreaming/Source/PixelStreaming/Private/VideoEncoderSingleLayerHardware.cpp
--- a/Engine/Plugins/Media/PixelStreaming/Source/PixelStreaming/Private/VideoEncoderSingleLayerHardware.cpp
+++ b/Engine/Plugins/Media/PixelStreaming/Source/PixelStreaming/Private/VideoEncoderSingleLayerHardware.cpp
@@ -36,8 +36,8 @@ namespace UE::PixelStreaming
 		VideoConfig.TargetFramerate = InCodecSettings->maxFramerate;
 		VideoConfig.TargetBitrate = InCodecSettings->startBitrate;
 		VideoConfig.MaxBitrate = InCodecSettings->maxBitrate;
-		VideoConfig.MinQP = PixelStreaming::Settings::CVarPixelStreamingEncoderMinQP.GetValueOnAnyThread();
-		VideoConfig.MaxQP = PixelStreaming::Settings::CVarPixelStreamingEncoderMaxQP.GetValueOnAnyThread();
+		VideoConfig.MinQuality = (1.0f - (PixelStreaming::Settings::CVarPixelStreamingEncoderMaxQP.GetValueOnAnyThread() / 100.0f)) * 51.0f;
+		VideoConfig.MaxQuality = (1.0f - (PixelStreaming::Settings::CVarPixelStreamingEncoderMinQP.GetValueOnAnyThread() / 100.0f)) * 51.0f;
 		VideoConfig.RateControlMode = PixelStreaming::Settings::GetRateControlCVar();
 		VideoConfig.bFillData = PixelStreaming::Settings::CVarPixelStreamingEnableFillerData.GetValueOnAnyThread();
 		VideoConfig.KeyframeInterval = PixelStreaming::Settings::CVarPixelStreamingEncoderKeyframeInterval.GetValueOnAnyThread();
@@ -50,35 +50,35 @@ namespace UE::PixelStreaming
 
 		switch (Codec)
 		{
-			case EPixelStreamingCodec::H264:
-			{
-				TUniquePtr<FVideoEncoderConfigH264> VideoConfig = MakeUnique<FVideoEncoderConfigH264>();
-				SetInitialSettings(InCodecSettings, *VideoConfig);
-				VideoConfig->Profile = PixelStreaming::Settings::GetH264Profile();
-				VideoConfig->RepeatSPSPPS = true;
-				VideoConfig->IntraRefreshPeriodFrames = PixelStreaming::Settings::CVarPixelStreamingEncoderIntraRefreshPeriodFrames.GetValueOnAnyThread();
-				VideoConfig->IntraRefreshCountFrames = PixelStreaming::Settings::CVarPixelStreamingEncoderIntraRefreshCountFrames.GetValueOnAnyThread();
-				VideoConfig->KeyframeInterval = PixelStreaming::Settings::CVarPixelStreamingEncoderKeyframeInterval.GetValueOnAnyThread();
-				// The WebRTC spec can only guarantee that the Baseline profile is supported. Therefore we use Baseline, but enable these extra
-				// features to improve bitrate usage
-				VideoConfig->AdaptiveTransformMode = EH264AdaptiveTransformMode::Enable;
-				VideoConfig->EntropyCodingMode = EH264EntropyCodingMode::CABAC;
-				InitialVideoConfig = MoveTemp(VideoConfig);
-				return WEBRTC_VIDEO_CODEC_OK;
-			}
-			case EPixelStreamingCodec::AV1:
-			{
-				TUniquePtr<FVideoEncoderConfigAV1> VideoConfig = MakeUnique<FVideoEncoderConfigAV1>();
-				SetInitialSettings(InCodecSettings, *VideoConfig);
-				VideoConfig->RepeatSeqHdr = true;
-				VideoConfig->IntraRefreshPeriodFrames = PixelStreaming::Settings::CVarPixelStreamingEncoderIntraRefreshPeriodFrames.GetValueOnAnyThread();
-				VideoConfig->IntraRefreshCountFrames = PixelStreaming::Settings::CVarPixelStreamingEncoderIntraRefreshCountFrames.GetValueOnAnyThread();
-				VideoConfig->KeyframeInterval = PixelStreaming::Settings::CVarPixelStreamingEncoderKeyframeInterval.GetValueOnAnyThread();
-				InitialVideoConfig = MoveTemp(VideoConfig);
-				return WEBRTC_VIDEO_CODEC_OK;
-			}
-			default:
-				return WEBRTC_VIDEO_CODEC_ERROR;
+		case EPixelStreamingCodec::H264:
+		{
+			TUniquePtr<FVideoEncoderConfigH264> VideoConfig = MakeUnique<FVideoEncoderConfigH264>();
+			SetInitialSettings(InCodecSettings, *VideoConfig);
+			VideoConfig->Profile = PixelStreaming::Settings::GetH264Profile();
+			VideoConfig->RepeatSPSPPS = true;
+			VideoConfig->IntraRefreshPeriodFrames = PixelStreaming::Settings::CVarPixelStreamingEncoderIntraRefreshPeriodFrames.GetValueOnAnyThread();
+			VideoConfig->IntraRefreshCountFrames = PixelStreaming::Settings::CVarPixelStreamingEncoderIntraRefreshCountFrames.GetValueOnAnyThread();
+			VideoConfig->KeyframeInterval = PixelStreaming::Settings::CVarPixelStreamingEncoderKeyframeInterval.GetValueOnAnyThread();
+			// The WebRTC spec can only guarantee that the Baseline profile is supported. Therefore we use Baseline, but enable these extra
+			// features to improve bitrate usage
+			VideoConfig->AdaptiveTransformMode = EH264AdaptiveTransformMode::Enable;
+			VideoConfig->EntropyCodingMode = EH264EntropyCodingMode::CABAC;
+			InitialVideoConfig = MoveTemp(VideoConfig);
+			return WEBRTC_VIDEO_CODEC_OK;
+		}
+		case EPixelStreamingCodec::AV1:
+		{
+			TUniquePtr<FVideoEncoderConfigAV1> VideoConfig = MakeUnique<FVideoEncoderConfigAV1>();
+			SetInitialSettings(InCodecSettings, *VideoConfig);
+			VideoConfig->RepeatSeqHdr = true;
+			VideoConfig->IntraRefreshPeriodFrames = PixelStreaming::Settings::CVarPixelStreamingEncoderIntraRefreshPeriodFrames.GetValueOnAnyThread();
+			VideoConfig->IntraRefreshCountFrames = PixelStreaming::Settings::CVarPixelStreamingEncoderIntraRefreshCountFrames.GetValueOnAnyThread();
+			VideoConfig->KeyframeInterval = PixelStreaming::Settings::CVarPixelStreamingEncoderKeyframeInterval.GetValueOnAnyThread();
+			InitialVideoConfig = MoveTemp(VideoConfig);
+			return WEBRTC_VIDEO_CODEC_OK;
+		}
+		default:
+			return WEBRTC_VIDEO_CODEC_ERROR;
 		}
 	}
 
@@ -151,18 +151,18 @@ namespace UE::PixelStreaming
 
 		switch (Codec)
 		{
-			case EPixelStreamingCodec::H264:
-			{
-				FVideoEncoderConfigH264& VideoConfig = *StaticCast<FVideoEncoderConfigH264*>(InitialVideoConfig.Get());
-				HardwareEncoder = Factory.GetOrCreateHardwareEncoder(StreamId, VideoConfig);
-				break;
-			}
-			case EPixelStreamingCodec::AV1:
-			{
-				FVideoEncoderConfigAV1& VideoConfig = *StaticCast<FVideoEncoderConfigAV1*>(InitialVideoConfig.Get());
-				HardwareEncoder = Factory.GetOrCreateHardwareEncoder(StreamId, VideoConfig);
-				break;
-			}
+		case EPixelStreamingCodec::H264:
+		{
+			FVideoEncoderConfigH264& VideoConfig = *StaticCast<FVideoEncoderConfigH264*>(InitialVideoConfig.Get());
+			HardwareEncoder = Factory.GetOrCreateHardwareEncoder(StreamId, VideoConfig);
+			break;
+		}
+		case EPixelStreamingCodec::AV1:
+		{
+			FVideoEncoderConfigAV1& VideoConfig = *StaticCast<FVideoEncoderConfigAV1*>(InitialVideoConfig.Get());
+			HardwareEncoder = Factory.GetOrCreateHardwareEncoder(StreamId, VideoConfig);
+			break;
+		}
 		}
 	}
 
@@ -199,17 +199,17 @@ namespace UE::PixelStreaming
 			TRACE_CPUPROFILER_EVENT_SCOPE_ON_CHANNEL_STR("PixelStreaming Hardware Encoding", PixelStreamingChannel);
 
 			uint64 CurrentFrameId = AdaptedLayer->Metadata.Id;
-			if(LastEncodedFrameId > 0)
+			if (LastEncodedFrameId > 0)
 			{
-				if(CurrentFrameId < LastEncodedFrameId)
+				if (CurrentFrameId < LastEncodedFrameId)
 				{
 					UE_LOG(LogPixelStreaming, Verbose, TEXT("Out of order frame entered hardware encoder."))
 				}
-				else if(CurrentFrameId == LastEncodedFrameId)
+				else if (CurrentFrameId == LastEncodedFrameId)
 				{
 					UE_LOG(LogPixelStreaming, Verbose, TEXT("Encoding the same frame."))
 				}
-				else if(CurrentFrameId - LastEncodedFrameId > 1)
+				else if (CurrentFrameId - LastEncodedFrameId > 1)
 				{
 					UE_LOG(LogPixelStreaming, Verbose, TEXT("Missed encoding %d frames"), CurrentFrameId - LastEncodedFrameId - 1);
 				}
@@ -272,17 +272,17 @@ namespace UE::PixelStreaming
 				webrtc::CodecSpecificInfo CodecInfo;
 				switch (Codec)
 				{
-					case EPixelStreamingCodec::H264:
-						CodecInfo.codecType = webrtc::VideoCodecType::kVideoCodecH264;
-						CodecInfo.codecSpecific.H264.packetization_mode = webrtc::H264PacketizationMode::NonInterleaved;
-						CodecInfo.codecSpecific.H264.temporal_idx = webrtc::kNoTemporalIdx;
-						CodecInfo.codecSpecific.H264.idr_frame = Packet.bIsKeyframe;
-						CodecInfo.codecSpecific.H264.base_layer_sync = false;
-
-						break;
-					case EPixelStreamingCodec::AV1:
-						CodecInfo.codecType = webrtc::VideoCodecType::kVideoCodecAV1;
-						break;
+				case EPixelStreamingCodec::H264:
+					CodecInfo.codecType = webrtc::VideoCodecType::kVideoCodecH264;
+					CodecInfo.codecSpecific.H264.packetization_mode = webrtc::H264PacketizationMode::NonInterleaved;
+					CodecInfo.codecSpecific.H264.temporal_idx = webrtc::kNoTemporalIdx;
+					CodecInfo.codecSpecific.H264.idr_frame = Packet.bIsKeyframe;
+					CodecInfo.codecSpecific.H264.base_layer_sync = false;
+
+					break;
+				case EPixelStreamingCodec::AV1:
+					CodecInfo.codecType = webrtc::VideoCodecType::kVideoCodecAV1;
+					break;
 				}
 
 #if PIXELSTREAMING_DUMP_ENCODING
@@ -344,24 +344,24 @@ namespace UE::PixelStreaming
 
 			switch (Codec)
 			{
-				case EPixelStreamingCodec::H264:
-					if (PinnedHardwareEncoder->GetInstance()->Has<FVideoEncoderConfigH264>())
-					{
-						FVideoEncoderConfigH264& VideoConfigH264 = PinnedHardwareEncoder->GetInstance()->Edit<FVideoEncoderConfigH264>();
-						VideoConfig = &VideoConfigH264;
+			case EPixelStreamingCodec::H264:
+				if (PinnedHardwareEncoder->GetInstance()->Has<FVideoEncoderConfigH264>())
+				{
+					FVideoEncoderConfigH264& VideoConfigH264 = PinnedHardwareEncoder->GetInstance()->Edit<FVideoEncoderConfigH264>();
+					VideoConfig = &VideoConfigH264;
 
-						VideoConfigH264.Profile = UE::PixelStreaming::Settings::GetH264Profile();
-					}
+					VideoConfigH264.Profile = UE::PixelStreaming::Settings::GetH264Profile();
+				}
 
-					break;
-				case EPixelStreamingCodec::AV1:
-					if (PinnedHardwareEncoder->GetInstance()->Has<FVideoEncoderConfigAV1>())
-					{
-						FVideoEncoderConfigAV1& VideoConfigAV1 = PinnedHardwareEncoder->GetInstance()->Edit<FVideoEncoderConfigAV1>();
-						VideoConfig = &VideoConfigAV1;
-					}
+				break;
+			case EPixelStreamingCodec::AV1:
+				if (PinnedHardwareEncoder->GetInstance()->Has<FVideoEncoderConfigAV1>())
+				{
+					FVideoEncoderConfigAV1& VideoConfigAV1 = PinnedHardwareEncoder->GetInstance()->Edit<FVideoEncoderConfigAV1>();
+					VideoConfig = &VideoConfigAV1;
+				}
 
-					break;
+				break;
 			}
 
 			if (PendingRateChange.IsSet())
@@ -380,16 +380,16 @@ namespace UE::PixelStreaming
 			// Change encoder settings through CVars
 			const int32 MaxBitrateCVar = UE::PixelStreaming::Settings::CVarPixelStreamingEncoderMaxBitrate.GetValueOnAnyThread();
 			const int32 TargetBitrateCVar = UE::PixelStreaming::Settings::CVarPixelStreamingEncoderTargetBitrate.GetValueOnAnyThread();
-			const int32 MinQPCVar = UE::PixelStreaming::Settings::CVarPixelStreamingEncoderMinQP.GetValueOnAnyThread();
 			const int32 MaxQPCVar = UE::PixelStreaming::Settings::CVarPixelStreamingEncoderMaxQP.GetValueOnAnyThread();
+			const int32 MinQPCVar = UE::PixelStreaming::Settings::CVarPixelStreamingEncoderMinQP.GetValueOnAnyThread();
 			const ERateControlMode RateControlCVar = UE::PixelStreaming::Settings::GetRateControlCVar();
 			const EMultipassMode MultiPassCVar = UE::PixelStreaming::Settings::GetMultipassCVar();
 			const bool bFillerDataCVar = UE::PixelStreaming::Settings::CVarPixelStreamingEnableFillerData.GetValueOnAnyThread();
 
 			VideoConfig->MaxBitrate = MaxBitrateCVar > -1 ? MaxBitrateCVar : VideoConfig->MaxBitrate;
 			VideoConfig->TargetBitrate = TargetBitrateCVar > -1 ? TargetBitrateCVar : WebRtcProposedTargetBitrate;
-			VideoConfig->MinQP = MinQPCVar;
-			VideoConfig->MaxQP = MaxQPCVar;
+			VideoConfig->MinQuality = (1.0f - (MaxQPCVar / 51.0f)) * 100.0f;
+			VideoConfig->MaxQuality = (1.0f - (MinQPCVar / 51.0f)) * 100.0f;
 			VideoConfig->RateControlMode = RateControlCVar;
 			VideoConfig->MultipassMode = MultiPassCVar;
 			VideoConfig->bFillData = bFillerDataCVar;
