[AVCodecs] Fix: Re-add external semaphore deletion as it was accidentally removed in CL29731260

diff --git a/Engine/Plugins/Experimental/AVCodecs/NVCodecs/Source/NVCodecs/Private/Video/Resources/VideoResourceCUDA.cpp b/Engine/Plugins/Experimental/AVCodecs/NVCodecs/Source/NVCodecs/Private/Video/Resources/VideoResourceCUDA.cpp
--- a/Engine/Plugins/Experimental/AVCodecs/NVCodecs/Source/NVCodecs/Private/Video/Resources/VideoResourceCUDA.cpp
+++ b/Engine/Plugins/Experimental/AVCodecs/NVCodecs/Source/NVCodecs/Private/Video/Resources/VideoResourceCUDA.cpp
@@ -134,6 +134,15 @@ FVideoResourceCUDA::~FVideoResourceCUDA()
 			FAVResult::Log(EAVResult::ErrorUnmapping, TEXT("Failed to clean up external memory"), TEXT("CUDA"), Result);
 		}
 	}
+
+	if (ExternalSemaphore != nullptr)
+	{
+		CUresult const Result = FCUDAModule::CUDA().cuDestroyExternalSemaphore(ExternalSemaphore);
+		if (Result != CUDA_SUCCESS)
+		{
+			FAVResult::Log(EAVResult::ErrorUnmapping, TEXT("Failed to destroy external semaphore"), TEXT("CUDA"), Result);
+		}
+	}
 }
 
 FAVResult FVideoResourceCUDA::Validate() const
