[AVCodecs] Fix: Update QP range for AV1 vs H.26X in NvEnc

diff --git a/Engine/Plugins/Experimental/AVCodecs/NVCodecs/Source/NVENC/Private/Video/Encoders/Configs/VideoEncoderConfigNVENC.cpp b/Engine/Plugins/Experimental/AVCodecs/NVCodecs/Source/NVENC/Private/Video/Encoders/Configs/VideoEncoderConfigNVENC.cpp
--- a/Engine/Plugins/Experimental/AVCodecs/NVCodecs/Source/NVENC/Private/Video/Encoders/Configs/VideoEncoderConfigNVENC.cpp
+++ b/Engine/Plugins/Experimental/AVCodecs/NVCodecs/Source/NVENC/Private/Video/Encoders/Configs/VideoEncoderConfigNVENC.cpp
@@ -119,8 +119,9 @@ DLLEXPORT FAVResult FAVExtension::TransformConfig(FVideoEncoderConfigNVENC& OutC
 		break;
 	}
 
-	uint32_t const MinQP = (1.0f - (InConfig.MaxQuality / 100.0f)) * 51.0f;
-	uint32_t const MaxQP = (1.0f - (InConfig.MinQuality / 100.0f)) * 51.0f;
+	float QPRange = !memcmp(&OutConfig.encodeGUID, &NV_ENC_CODEC_AV1_GUID, sizeof(GUID)) ? 255.0f : 51.0f;
+	uint32_t const MinQP = (1.0f - (InConfig.MaxQuality / 100.0f)) * QPRange;
+	uint32_t const MaxQP = (1.0f - (InConfig.MinQuality / 100.0f)) * QPRange;
 
 	ERateControlMode ActualRateControlMode = InConfig.Preset == EAVPreset::Lossless ? ERateControlMode::ConstQP : InConfig.RateControlMode;
 	TAVResult<NV_ENC_PARAMS_RC_MODE> const ConvertedRateControlMode = FVideoEncoderConfigNVENC::ConvertRateControlMode(ActualRateControlMode);
@@ -183,8 +184,9 @@ DLLEXPORT FAVResult FAVExtension::TransformConfig(FVideoEncoderConfig& OutConfig
 	OutConfig.Height = InConfig.encodeHeight;
 
 	NV_ENC_RC_PARAMS& RateControlParams = InConfig.encodeConfig->rcParams;
-	OutConfig.MinQuality = static_cast<int32>((1.0f - (RateControlParams.maxQP.qpIntra / 51.0f)) * 100.0f);
-	OutConfig.MaxQuality = static_cast<int32>((1.0f - (RateControlParams.minQP.qpIntra / 51.0f)) * 100.0f);
+	float QPRange = !memcmp(&InConfig.encodeGUID, &NV_ENC_CODEC_AV1_GUID, sizeof(GUID)) ? 255.0f : 51.0f;
+	OutConfig.MinQuality = static_cast<int32>((1.0f - (RateControlParams.maxQP.qpIntra / QPRange)) * 100.0f);
+	OutConfig.MaxQuality = static_cast<int32>((1.0f - (RateControlParams.minQP.qpIntra / QPRange)) * 100.0f);
 	OutConfig.TargetBitrate = RateControlParams.averageBitRate == DEFAULT_BITRATE_TARGET ? -1 : RateControlParams.averageBitRate;
 	OutConfig.MaxBitrate = RateControlParams.maxBitRate == DEFAULT_BITRATE_MAX ? -1 : RateControlParams.maxBitRate;
 
