[AVCodecs, PixelStreaming2] Fix: Inconsistent NvEnc bitrate causing increased latency

diff --git a/Engine/Plugins/Experimental/AVCodecs/NVCodecs/Source/NVENC/Private/Video/Encoders/Configs/VideoEncoderConfigNVENC.cpp b/Engine/Plugins/Experimental/AVCodecs/NVCodecs/Source/NVENC/Private/Video/Encoders/Configs/VideoEncoderConfigNVENC.cpp
--- a/Engine/Plugins/Experimental/AVCodecs/NVCodecs/Source/NVENC/Private/Video/Encoders/Configs/VideoEncoderConfigNVENC.cpp
+++ b/Engine/Plugins/Experimental/AVCodecs/NVCodecs/Source/NVENC/Private/Video/Encoders/Configs/VideoEncoderConfigNVENC.cpp
@@ -146,6 +146,12 @@ DLLEXPORT FAVResult FAVExtension::TransformConfig(FVideoEncoderConfigNVENC& OutC
 	RateControlParams.constQP = RateControlParams.maxQP;
 	RateControlParams.enableMinQP = InConfig.MaxQuality < 100;
 	RateControlParams.enableMaxQP = InConfig.MinQuality > 0;
+	RateControlParams.enableAQ = 1;
+	if (ActualRateControlMode == ERateControlMode::CBR)
+	{	
+		RateControlParams.vbvBufferSize = RateControlParams.averageBitRate / (InConfig.TargetFramerate > 0 ? InConfig.TargetFramerate : 60);		
+		RateControlParams.vbvInitialDelay = RateControlParams.vbvBufferSize;
+	}
 
 	// IDR period - how often to send IDR (instantaneous decode refresh) frames, a.k.a keyframes. This can stabilise a stream that dropped/lost some frames (but at the cost of more bandwidth)
 	if (!FMemory::Memcmp(&OutConfig.encodeGUID, &NV_ENC_CODEC_H264_GUID, sizeof(GUID)))
diff --git a/Engine/Plugins/Experimental/AVCodecs/NVCodecs/Source/NVENC/Public/Video/Encoders/Configs/VideoEncoderConfigNVENC.h b/Engine/Plugins/Experimental/AVCodecs/NVCodecs/Source/NVENC/Public/Video/Encoders/Configs/VideoEncoderConfigNVENC.h
--- a/Engine/Plugins/Experimental/AVCodecs/NVCodecs/Source/NVENC/Public/Video/Encoders/Configs/VideoEncoderConfigNVENC.h
+++ b/Engine/Plugins/Experimental/AVCodecs/NVCodecs/Source/NVENC/Public/Video/Encoders/Configs/VideoEncoderConfigNVENC.h
@@ -25,7 +25,7 @@ public:
 		, NV_ENC_INITIALIZE_PARAMS({})
 	{
 		version = NV_ENC_INITIALIZE_PARAMS_VER;
-		presetGUID = NV_ENC_PRESET_P4_GUID;
+		presetGUID = NV_ENC_PRESET_P1_GUID;
 		frameRateNum = 60;
 		frameRateDen = 1;
 		enablePTD = 1;
diff --git a/Engine/Plugins/Media/PixelStreaming2/Source/PixelStreaming2RTC/Private/EpicRtcVideoEncoder.cpp b/Engine/Plugins/Media/PixelStreaming2/Source/PixelStreaming2RTC/Private/EpicRtcVideoEncoder.cpp
--- a/Engine/Plugins/Media/PixelStreaming2/Source/PixelStreaming2RTC/Private/EpicRtcVideoEncoder.cpp
+++ b/Engine/Plugins/Media/PixelStreaming2/Source/PixelStreaming2RTC/Private/EpicRtcVideoEncoder.cpp
@@ -232,7 +232,7 @@ namespace UE::PixelStreaming2
 			{
 				for (size_t ti = 0; ti < Video::MaxTemporalStreams; ++ti)
 				{
-					EpicRtcTargetBitrates[si][ti].Emplace(VideoEncoderConfig._rateControl._targetBitrate->GetBitrate(si, ti));
+					EpicRtcTargetBitrates[si][ti].Emplace(VideoEncoderConfig._rateControl._bitrate->GetBitrate(si, ti));
 				}
 			}
 
