[AVCodecs] Fix: NvEnc not working with D3D12 placed resources

diff --git a/Engine/Plugins/Experimental/AVCodecs/NVCodecs/Source/NVENC/Private/NVENCModule.cpp b/Engine/Plugins/Experimental/AVCodecs/NVCodecs/Source/NVENC/Private/NVENCModule.cpp
--- a/Engine/Plugins/Experimental/AVCodecs/NVCodecs/Source/NVENC/Private/NVENCModule.cpp
+++ b/Engine/Plugins/Experimental/AVCodecs/NVCodecs/Source/NVENC/Private/NVENCModule.cpp
@@ -77,9 +77,6 @@ public:
 			::RegisterPermutationsOf<FVideoEncoderNVENCCUDA>
 			::With<
 				FVideoResourceCUDA,
-#if PLATFORM_WINDOWS
-				FVideoResourceD3D12,
-#endif
 				FVideoResourceVulkan>
 			::And<FVideoEncoderConfigNVENC, FVideoEncoderConfigH264>(
 				[CheckCodecSupport](TSharedRef<FAVDevice> const& NewDevice, TSharedRef<FAVInstance> const& NewInstance)
@@ -95,9 +92,6 @@ public:
 			::RegisterPermutationsOf<FVideoEncoderNVENCCUDA>
 			::With<
 				FVideoResourceCUDA,
-#if PLATFORM_WINDOWS
-				FVideoResourceD3D12,
-#endif
 				FVideoResourceVulkan>
 			::And<FVideoEncoderConfigNVENC, FVideoEncoderConfigH265>(
 				[CheckCodecSupport](TSharedRef<FAVDevice> const& NewDevice, TSharedRef<FAVInstance> const& NewInstance)
@@ -113,9 +107,6 @@ public:
 			::RegisterPermutationsOf<FVideoEncoderNVENCCUDA>
 			::With<
 				FVideoResourceCUDA,
-#if PLATFORM_WINDOWS
-				FVideoResourceD3D12,
-#endif
 				FVideoResourceVulkan>
 			::And<FVideoEncoderConfigNVENC, FVideoEncoderConfigAV1>(
 				[CheckCodecSupport](TSharedRef<FAVDevice> const& NewDevice, TSharedRef<FAVInstance> const& NewInstance)
@@ -165,6 +156,42 @@ public:
 
 					return FAPI::Get<FNVENC>().IsValid() && NewDevice->HasContext<FVideoContextD3D11>() && bSupportsCodec;
 				});
+
+		FVideoEncoder
+			::RegisterPermutationsOf<FVideoEncoderNVENCD3D12>
+			::With<FVideoResourceD3D12>
+			::And<FVideoEncoderConfigNVENC, FVideoEncoderConfigH264>(
+				[CheckCodecSupport](TSharedRef<FAVDevice> const& NewDevice, TSharedRef<FAVInstance> const& NewInstance)
+				{
+					// static local because we don't want to constantly create new encoder sessions every time we query codec support
+					static bool bSupportsCodec = CheckCodecSupport(NewDevice, NV_ENC_CODEC_H264_GUID);
+
+					return FAPI::Get<FNVENC>().IsValid() && NewDevice->HasContext<FVideoContextD3D12>() && bSupportsCodec;
+				});
+
+		FVideoEncoder
+			::RegisterPermutationsOf<FVideoEncoderNVENCD3D12>
+			::With<FVideoResourceD3D12>
+			::And<FVideoEncoderConfigNVENC, FVideoEncoderConfigH265>(
+				[CheckCodecSupport](TSharedRef<FAVDevice> const& NewDevice, TSharedRef<FAVInstance> const& NewInstance)
+				{
+					// static local because we don't want to constantly create new encoder sessions every time we query codec support
+					static bool bSupportsCodec = CheckCodecSupport(NewDevice, NV_ENC_CODEC_HEVC_GUID);
+
+					return FAPI::Get<FNVENC>().IsValid() && NewDevice->HasContext<FVideoContextD3D12>() && bSupportsCodec;
+				});
+
+		FVideoEncoder
+			::RegisterPermutationsOf<FVideoEncoderNVENCD3D12>
+			::With<FVideoResourceD3D12>
+			::And<FVideoEncoderConfigNVENC, FVideoEncoderConfigAV1>(
+				[CheckCodecSupport](TSharedRef<FAVDevice> const& NewDevice, TSharedRef<FAVInstance> const& NewInstance)
+				{
+					// static local because we don't want to constantly create new encoder sessions every time we query codec support
+					static bool bSupportsCodec = CheckCodecSupport(NewDevice, NV_ENC_CODEC_AV1_GUID);
+
+					return FAPI::Get<FNVENC>().IsValid() && NewDevice->HasContext<FVideoContextD3D12>() && bSupportsCodec;
+				});
 #endif // PLATFORM_WINDOWS
 	
 	}
diff --git a/Engine/Plugins/Experimental/AVCodecs/NVCodecs/Source/NVENC/Private/Video/Encoders/VideoEncoderNVENC.cpp b/Engine/Plugins/Experimental/AVCodecs/NVCodecs/Source/NVENC/Private/Video/Encoders/VideoEncoderNVENC.cpp
--- a/Engine/Plugins/Experimental/AVCodecs/NVCodecs/Source/NVENC/Private/Video/Encoders/VideoEncoderNVENC.cpp
+++ b/Engine/Plugins/Experimental/AVCodecs/NVCodecs/Source/NVENC/Private/Video/Encoders/VideoEncoderNVENC.cpp
@@ -1,95 +1,48 @@
 // Copyright Epic Games, Inc. All Rights Reserved.
 
 #include "Video/Encoders/VideoEncoderNVENC.h"
+#include "Templates/AlignmentTemplates.h"
 
-FEncoderNVENC::~FEncoderNVENC()
+FAVResult FEncoderNVENC::Open(TSharedRef<FAVDevice> const& NewDevice, TSharedRef<FAVInstance> const& NewInstance, NV_ENC_OPEN_ENCODE_SESSION_EX_PARAMS& SessionParams)
 {
 	Close();
-}
-
-bool FEncoderNVENC::IsOpen() const
-{
-	return Encoder != nullptr;
-}
 
-FAVResult FEncoderNVENC::ReOpen()
-{
-	Close();
-
-	NV_ENC_STRUCT(NV_ENC_OPEN_ENCODE_SESSION_EX_PARAMS, SessionParams);
-	SessionParams.apiVersion = NVENCAPI_VERSION;
-
-	SessionParams.deviceType = SessionDeviceType;
-	SessionParams.device = SessionDevice;
+	SessionDeviceType = SessionParams.deviceType;
+	SessionDevice = SessionParams.device;
 
 	NVENCSTATUS const Result = FAPI::Get<FNVENC>().nvEncOpenEncodeSessionEx(&const_cast<NV_ENC_OPEN_ENCODE_SESSION_EX_PARAMS&>(SessionParams), &Encoder);
 	if (Result != NV_ENC_SUCCESS)
 	{
 		Close();
 
-		return FAVResult(EAVResult::ErrorCreating, TEXT("Failed to re-create encoder"), TEXT("NVENC"), Result);
+		return FAVResult(EAVResult::ErrorCreating, TEXT("Failed to create encoder"), TEXT("NVENC"), Result);
 	}
 
 	return EAVResult::Success;
 }
 
-FAVResult FEncoderNVENC::Open(TSharedRef<FAVDevice> const& NewDevice, TSharedRef<FAVInstance> const& NewInstance, TFunction<void(NV_ENC_OPEN_ENCODE_SESSION_EX_PARAMS&)> SetupEncoderSessionFunc)
+FAVResult FEncoderNVENC::ReOpen()
 {
 	Close();
 
 	NV_ENC_STRUCT(NV_ENC_OPEN_ENCODE_SESSION_EX_PARAMS, SessionParams);
 	SessionParams.apiVersion = NVENCAPI_VERSION;
 
-	SetupEncoderSessionFunc(SessionParams);
-
-	SessionDeviceType = SessionParams.deviceType;
-	SessionDevice = SessionParams.device;
+	SessionParams.deviceType = SessionDeviceType;
+	SessionParams.device = SessionDevice;
 
 	NVENCSTATUS const Result = FAPI::Get<FNVENC>().nvEncOpenEncodeSessionEx(&const_cast<NV_ENC_OPEN_ENCODE_SESSION_EX_PARAMS&>(SessionParams), &Encoder);
 	if (Result != NV_ENC_SUCCESS)
 	{
 		Close();
 
-		return FAVResult(EAVResult::ErrorCreating, TEXT("Failed to create encoder"), TEXT("NVENC"), Result);
+		return FAVResult(EAVResult::ErrorCreating, TEXT("Failed to re-create encoder"), TEXT("NVENC-CUDA"), Result);
 	}
 
 	return EAVResult::Success;
 }
 
-void FEncoderNVENC::Close()
-{
-	if (IsOpen())
-	{
-		if (Buffer != nullptr)
-		{
-			NVENCSTATUS const Result = FAPI::Get<FNVENC>().nvEncDestroyBitstreamBuffer(Encoder, Buffer);
-			if (Result != NV_ENC_SUCCESS)
-			{
-				FAVResult::Log(EAVResult::ErrorDestroying, TEXT("Failed to destroy output buffer"), TEXT("NVENC"), Result);
-			}
-
-			Buffer = nullptr;
-		}
-
-		if (Encoder != nullptr)
-		{
-			NVENCSTATUS const Result = FAPI::Get<FNVENC>().nvEncDestroyEncoder(Encoder);
-			if (Result != NV_ENC_SUCCESS)
-			{
-				FAVResult::Log(EAVResult::ErrorDestroying, TEXT("Failed to destroy encoder"), TEXT("NVENC"), Result);
-			}
-
-			Encoder = nullptr;
-		}
-	}
-}
-
-bool FEncoderNVENC::IsInitialized() const
-{
-	return IsOpen() && Buffer != nullptr;
-}
-
-FAVResult FEncoderNVENC::ApplyConfig(FVideoEncoderConfigNVENC const& AppliedConfig, FVideoEncoderConfigNVENC& PendingConfig, TFunction<FAVResult()> ApplyConfigFunc)
+FAVResult FEncoderNVENC::ApplyConfig(FVideoEncoderConfigNVENC const& AppliedConfig, FVideoEncoderConfigNVENC& PendingConfig)
 {
 	if (IsOpen())
 	{
@@ -145,25 +98,25 @@ FAVResult FEncoderNVENC::ApplyConfig(FVideoEncoderConfigNVENC const& AppliedConf
 
 			if (!IsInitialized())
 			{
-				NVENCSTATUS Result = FAPI::Get<FNVENC>().nvEncInitializeEncoder(Encoder, &const_cast<FVideoEncoderConfigNVENC&>(PendingConfig));
-				if (Result != NV_ENC_SUCCESS)
 				{
-					return FAVResult(EAVResult::ErrorCreating, TEXT("Failed to initialize encoder"), TEXT("NVENC"), Result);
+					NVENCSTATUS Result = FAPI::Get<FNVENC>().nvEncInitializeEncoder(Encoder, &const_cast<FVideoEncoderConfigNVENC&>(PendingConfig));
+					if (Result != NV_ENC_SUCCESS)
+					{
+						return FAVResult(EAVResult::ErrorCreating, TEXT("Failed to initialize encoder"), TEXT("NVENC"), Result);
+					}
 				}
 
-				NV_ENC_STRUCT(NV_ENC_CREATE_BITSTREAM_BUFFER, CreateBuffer);
-
-				Result = FAPI::Get<FNVENC>().nvEncCreateBitstreamBuffer(Encoder, &CreateBuffer);
-				if (Result != NV_ENC_SUCCESS)
 				{
-					return FAVResult(EAVResult::ErrorCreating, TEXT("Failed to create output buffer"), TEXT("NVENC"), Result);
+					FAVResult const Result = OnPostInitialization(PendingConfig);
+					if (Result.IsNotSuccess())
+					{
+						return Result;
+					}
 				}
-
-				Buffer = CreateBuffer.bitstreamBuffer;
 			}
 		}
 
-		return ApplyConfigFunc();
+		return EAVResult::Success;
 	}
 
 	return FAVResult(EAVResult::ErrorInvalidState, TEXT("Encoder not open"), TEXT("NVENC"));
@@ -211,132 +164,12 @@ void FEncoderNVENC::SetMaxResolution(FVideoEncoderConfigNVENC& PendingConfig)
 	{
 		GetMaxDeviceEncodeResolution(PendingConfig);
 
-		PendingConfig.maxEncodeWidth = FMath::Min(PendingConfig.maxEncodeWidth, MaxDeviceEncodeWidth);
-		PendingConfig.maxEncodeHeight = FMath::Min(PendingConfig.maxEncodeHeight, MaxDeviceEncodeHeight);
+		PendingConfig.maxEncodeWidth = bHasMaxDeviceResolution ? MaxDeviceEncodeWidth : PendingConfig.maxEncodeWidth;
+		PendingConfig.maxEncodeHeight = bHasMaxDeviceResolution ? MaxDeviceEncodeHeight : PendingConfig.maxEncodeHeight;
 	}
 }
 
-FAVResult FEncoderNVENC::SendFrame(TSharedPtr<FVideoResource> const& Resource, uint32 Timestamp, bool bForceKeyframe, TFunction<FAVResult()> ApplyConfigFunc, TFunction<void(NV_ENC_REGISTER_RESOURCE&)> SetResourceToRegisterFunc)
-{
-	if (IsOpen())
-	{
-		FAVResult AVResult = ApplyConfigFunc();
-		if (AVResult.IsNotSuccess())
-		{
-			return AVResult;
-		}
-
-		NV_ENC_STRUCT(NV_ENC_MAP_INPUT_RESOURCE, MapResource);
-		NV_ENC_STRUCT(NV_ENC_REGISTER_RESOURCE, RegisterResource);
-		NV_ENC_STRUCT(NV_ENC_PIC_PARAMS, Picture);
-
-		NVENCSTATUS Result;
-		if (Resource.IsValid())
-		{
-			TAVResult<NV_ENC_BUFFER_FORMAT> const ConvertedPixelFormat = FVideoEncoderConfigNVENC::ConvertFormat(Resource->GetFormat());
-			if (ConvertedPixelFormat.IsNotSuccess())
-			{
-				return ConvertedPixelFormat;
-			}
-
-			// START DEBUG
-#if DEBUG_DUMP_TO_DISK
-			{
-				TArray64<uint8> OutData;
-
-				Resource->ReadData(OutData);
-
-				FString SaveName = FString::Printf(TEXT("%s/DumpInput/image%05d.%s"), *FPaths::ProjectSavedDir(), Timestamp, ConvertedPixelFormat == NV_ENC_BUFFER_FORMAT_NV12 ? TEXT("nv12") : TEXT("p016"));
-
-				FFileHelper::SaveArrayToFile(OutData, *SaveName);
-			}
-#endif
-			// END DEBUG
-
-			SetResourceToRegisterFunc(RegisterResource);
-			RegisterResource.width = Resource->GetWidth();
-			RegisterResource.height = Resource->GetHeight();
-			RegisterResource.pitch = Resource->GetStride();
-			RegisterResource.bufferFormat = ConvertedPixelFormat;
-			RegisterResource.bufferUsage = NV_ENC_INPUT_IMAGE;
-
-			Result = FAPI::Get<FNVENC>().nvEncRegisterResource(Encoder, &RegisterResource);
-			if (Result != NV_ENC_SUCCESS)
-			{
-				return FAVResult(EAVResult::ErrorMapping, TEXT("Failed to register frame resource"), TEXT("NVENC"), Result);
-			}
-
-			MapResource.registeredResource = RegisterResource.registeredResource;
-
-			Result = FAPI::Get<FNVENC>().nvEncMapInputResource(Encoder, &MapResource);
-			if (Result != NV_ENC_SUCCESS)
-			{
-				return FAVResult(EAVResult::ErrorMapping, TEXT("Failed to map frame resource"), TEXT("NVENC"), Result);
-			}
-
-			Picture.inputWidth = RegisterResource.width;
-			Picture.inputHeight = RegisterResource.height;
-			Picture.inputPitch = RegisterResource.pitch;
-			Picture.inputBuffer = MapResource.mappedResource;
-			Picture.bufferFmt = MapResource.mappedBufferFmt;
-			Picture.pictureStruct = NV_ENC_PIC_STRUCT_FRAME;
-			Picture.inputTimeStamp = Timestamp;
-
-			if (bForceKeyframe)
-			{
-				Picture.encodePicFlags |= NV_ENC_PIC_FLAG_FORCEIDR;
-			}
-		}
-		else
-		{
-			Picture.encodePicFlags = NV_ENC_PIC_FLAG_EOS;
-		}
-
-		if (Resource.IsValid())
-		{
-			Resource->Lock();
-		}
-
-		AVResult = SendFrame(Picture);
-
-		if (Resource.IsValid())
-		{
-			Resource->Unlock();
-		}
-
-		if (MapResource.mappedResource != nullptr)
-		{
-			Result = FAPI::Get<FNVENC>().nvEncUnmapInputResource(Encoder, MapResource.mappedResource);
-			if (Result != NV_ENC_SUCCESS)
-			{
-				FAVResult const UnmapResult = FAVResult(EAVResult::ErrorUnmapping, TEXT("Failed to unmap frame resource"), TEXT("NVENC"), Result);
-				if (UnmapResult < AVResult)
-				{
-					AVResult = UnmapResult;
-				}
-			}
-		}
-
-		if (RegisterResource.registeredResource != nullptr)
-		{
-			Result = FAPI::Get<FNVENC>().nvEncUnregisterResource(Encoder, RegisterResource.registeredResource);
-			if (Result != NV_ENC_SUCCESS)
-			{
-				FAVResult const UnmapResult = FAVResult(EAVResult::ErrorUnmapping, TEXT("Failed to unregister frame resource"), TEXT("NVENC"), Result);
-				if (UnmapResult < AVResult)
-				{
-					AVResult = UnmapResult;
-				}
-			}
-		}
-
-		return AVResult;
-	}
-
-	return FAVResult(EAVResult::ErrorInvalidState, TEXT("Encoder not open"), TEXT("NVENC"));
-}
-
-FAVResult FEncoderNVENC::SendFrame(NV_ENC_PIC_PARAMS Input)
+FAVResult FEncoderNVENC::SendFrame(NV_ENC_PIC_PARAMS& Input, NV_ENC_LOCK_BITSTREAM& BitstreamLock)
 {
 	if (IsOpen())
 	{
@@ -352,10 +185,6 @@ FAVResult FEncoderNVENC::SendFrame(NV_ENC_PIC_PARAMS Input)
 		}
 		else
 		{
-			if (Input.outputBitstream == nullptr)
-			{
-				Input.outputBitstream = Buffer;
-			}
 
 			NVENCSTATUS Result = FAPI::Get<FNVENC>().nvEncEncodePicture(Encoder, &Input);
 			if (Result != NV_ENC_SUCCESS)
@@ -363,9 +192,6 @@ FAVResult FEncoderNVENC::SendFrame(NV_ENC_PIC_PARAMS Input)
 				return FAVResult(EAVResult::Error, TEXT("Error encoding picture"), TEXT("NVENC"), Result);
 			}
 
-			NV_ENC_STRUCT(NV_ENC_LOCK_BITSTREAM, BitstreamLock);
-			BitstreamLock.outputBitstream = Buffer;
-
 			Result = FAPI::Get<FNVENC>().nvEncLockBitstream(Encoder, &BitstreamLock);
 			if (Result != NV_ENC_SUCCESS)
 			{
@@ -397,28 +223,86 @@ FAVResult FEncoderNVENC::SendFrame(NV_ENC_PIC_PARAMS Input)
 	return FAVResult(EAVResult::ErrorInvalidState, TEXT("Encoder not open"), TEXT("NVENC"));
 }
 
-#if PLATFORM_WINDOWS
-FAVResult FEncoderNVENC::SendFrameD3D11(TRefCountPtr<ID3D11Device> Device, TSharedPtr<FVideoResourceD3D11> const& Resource, uint32 Timestamp, bool bForceKeyframe, TFunction<FAVResult()> ApplyConfigFunc)
+FAVResult FEncoderNVENC::ReceivePacket(FVideoPacket& OutPacket)
+{
+	if (IsOpen())
+	{
+		if (Packets.Dequeue(OutPacket))
+		{
+			return EAVResult::Success;
+		}
+		
+		return EAVResult::PendingInput;
+	}
+	
+	return FAVResult(EAVResult::ErrorInvalidState, TEXT("Encoder not open"), TEXT("NVENC"));
+}
+
+FAVResult FEncoderNVENC::RegisterResource(NV_ENC_REGISTER_RESOURCE& OutRegisterResource, NV_ENC_INPUT_RESOURCE_TYPE ResourceType, NV_ENC_BUFFER_FORMAT ResourceFormat, void* Resource, uint32 Width, uint32 Height, uint32 Stride, NV_ENC_BUFFER_USAGE BufferUsage)
+{
+	OutRegisterResource.resourceType = ResourceType;
+	OutRegisterResource.resourceToRegister = Resource;
+	OutRegisterResource.width = Width;
+	OutRegisterResource.height = Height;
+	OutRegisterResource.pitch = Stride;
+	OutRegisterResource.bufferFormat = ResourceFormat;
+	OutRegisterResource.bufferUsage = BufferUsage;
+
+	NVENCSTATUS Result = FAPI::Get<FNVENC>().nvEncRegisterResource(Encoder, &OutRegisterResource);
+	if (Result != NV_ENC_SUCCESS)
+	{
+		return FAVResult(EAVResult::ErrorMapping, TEXT("Failed to register frame resource"), TEXT("NVENC"), Result);
+	}
+
+	return EAVResult::Success;
+}
+
+FAVResult FEncoderNVENC::UnregisterResource(NV_ENC_REGISTERED_PTR RegisteredResource)
+{
+	if (RegisteredResource != nullptr)
+	{
+		NVENCSTATUS Result = FAPI::Get<FNVENC>().nvEncUnregisterResource(Encoder, RegisteredResource);
+		if (Result != NV_ENC_SUCCESS)
+		{
+			return FAVResult(EAVResult::ErrorUnmapping, TEXT("Failed to unregister frame resource"), TEXT("NVENC"), Result);
+		}
+	}
+
+	return EAVResult::Success;
+}
+
+FAVResult FEncoderNVENC::MapResource(NV_ENC_MAP_INPUT_RESOURCE& OutMapResource, NV_ENC_REGISTERED_PTR RegisteredResource)
 {
-	TRefCountPtr<ID3D11Texture2D>& Tex = const_cast<TRefCountPtr<ID3D11Texture2D>&>(Resource->GetRaw());
+	OutMapResource.registeredResource = RegisteredResource;
 
-	// Very important: We create a new device for NVENC - we must share out D3D11 texture with the new device by "Opening" it to that device.
-	HRESULT Result = Device->OpenSharedResource(Resource->GetSharedHandle(), __uuidof(ID3D11Texture2D), (void**)(Tex.GetInitReference()));
-	if (Result != S_OK)
+	NVENCSTATUS Result = FAPI::Get<FNVENC>().nvEncMapInputResource(Encoder, &OutMapResource);
+	if (Result != NV_ENC_SUCCESS)
 	{
-		return FAVResult(EAVResult::Fatal, TEXT("Failed to open shared handle."), TEXT("NVENC"), Result);
+		return FAVResult(EAVResult::ErrorMapping, TEXT("Failed to map frame resource"), TEXT("NVENC"), Result);
 	}
 
-	return SendFrame(Resource, Timestamp, bForceKeyframe, ApplyConfigFunc, [Resource](NV_ENC_REGISTER_RESOURCE& RegisterResource) {
-		RegisterResource.resourceType = NV_ENC_INPUT_RESOURCE_TYPE_DIRECTX;
-		RegisterResource.resourceToRegister = Resource->GetRaw();
-	});
+	return EAVResult::Success;
 }
 
+FAVResult FEncoderNVENC::UnmapResource(NV_ENC_INPUT_PTR MappedResource)
+{
+	if (MappedResource != nullptr)
+	{
+		NVENCSTATUS Result = FAPI::Get<FNVENC>().nvEncUnmapInputResource(Encoder, MappedResource);
+		if (Result != NV_ENC_SUCCESS)
+		{
+			return FAVResult(EAVResult::ErrorUnmapping, TEXT("Failed to unmap frame resource"), TEXT("NVENC"), Result);
+		}
+	}
+
+	return EAVResult::Success;
+}
+
+#if PLATFORM_WINDOWS
 /**
  * @returns The FString version of HRESULT error message so we can log it.
  **/
-const FString AVGetComErrorDescription(HRESULT Res)
+const FString AVUtils::AVGetComErrorDescription(HRESULT Res)
 {
 	const uint32 BufSize = 4096;
 	WIDECHAR	 Buffer[4096];
@@ -437,73 +321,4 @@ const FString AVGetComErrorDescription(HRESULT Res)
 		return TEXT("[cannot find error description]");
 	}
 }
-
-/**
- * You may be asking yourself, why on Earth would we need to create a seperate D3D11 device when using NVENC?
- * Well, it basically comes down to the D3D11 device created by Unreal Engine is not compatible with NVENC due to sdk version changes (we think?)
- * and Unreal Engine is in no hurry to bump D3D11 version - so we create out own device. This works only because
- * the D3D11 textures we are using are created as shared resources and accessed through shared handles.
- * Without this chicainery the D3D11 device will be ejected and Unreal Engine and NVENC will crash.
- */
-FAVResult FEncoderNVENC::CreateD3D11Device(TSharedRef<FAVDevice> const& InDevice, TRefCountPtr<ID3D11Device>& OutEncoderDevice, TRefCountPtr<ID3D11DeviceContext>& OutEncoderDeviceContext)
-{
-	TRefCountPtr<IDXGIDevice>  DXGIDevice;
-	TRefCountPtr<IDXGIAdapter> Adapter;
-
-	HRESULT Result = InDevice->GetContext<FVideoContextD3D11>()->Device->QueryInterface(__uuidof(IDXGIDevice), (void**)DXGIDevice.GetInitReference());
-	if (Result != S_OK)
-	{
-		return FAVResult(EAVResult::Fatal, FString::Printf(TEXT("ID3D11Device::QueryInterface() failed 0x%X - %s."), (uint32)Result, *AVGetComErrorDescription(Result)), TEXT("D3D11"));
-	}
-	else if ((Result = DXGIDevice->GetAdapter(Adapter.GetInitReference())) != S_OK)
-	{
-		return FAVResult(EAVResult::Fatal, FString::Printf(TEXT("DXGIDevice::GetAdapter() failed 0x%X - %s."), (uint32)Result, *AVGetComErrorDescription(Result)), TEXT("D3D11"));
-	}
-
-	uint32			  DeviceFlags = 0;
-	D3D_FEATURE_LEVEL FeatureLevel = D3D_FEATURE_LEVEL_11_0;
-	D3D_FEATURE_LEVEL ActualFeatureLevel;
-
-	if ((Result = D3D11CreateDevice(
-			 Adapter,
-			 D3D_DRIVER_TYPE_UNKNOWN,
-			 NULL,
-			 DeviceFlags,
-			 &FeatureLevel,
-			 1,
-			 D3D11_SDK_VERSION,
-			 OutEncoderDevice.GetInitReference(),
-			 &ActualFeatureLevel,
-			 OutEncoderDeviceContext.GetInitReference()))
-		!= S_OK)
-	{
-		return FAVResult(EAVResult::Fatal, FString::Printf(TEXT("D3D11CreateDevice() failed 0x%X - %s."), (uint32)Result, *AVGetComErrorDescription(Result)), TEXT("D3D11"));
-	}
-
-	return FAVResult(EAVResult::Success, TEXT("Created D3D11 device for NVENC."), TEXT("D3D11"));
-}
-
-#endif // PLATFORM_WINDOWS
-
-FAVResult FEncoderNVENC::SendFrameCUDA(TSharedPtr<FVideoResourceCUDA> const& Resource, uint32 Timestamp, bool bForceKeyframe, TFunction<FAVResult()> ApplyConfigFunc)
-{
-	return SendFrame(Resource, Timestamp, bForceKeyframe, ApplyConfigFunc, [Resource](NV_ENC_REGISTER_RESOURCE& RegisterResource) {
-		RegisterResource.resourceType = NV_ENC_INPUT_RESOURCE_TYPE_CUDAARRAY;
-		RegisterResource.resourceToRegister = Resource->GetRaw();
-	});
-}
-
-FAVResult FEncoderNVENC::ReceivePacket(FVideoPacket& OutPacket)
-{
-	if (IsOpen())
-	{
-		if (Packets.Dequeue(OutPacket))
-		{
-			return EAVResult::Success;
-		}
-
-		return EAVResult::PendingInput;
-	}
-
-	return FAVResult(EAVResult::ErrorInvalidState, TEXT("Encoder not open"), TEXT("NVENC"));
-}
+#endif
\ No newline at end of file
diff --git a/Engine/Plugins/Experimental/AVCodecs/NVCodecs/Source/NVENC/Private/Video/Encoders/VideoEncoderNVENCCUDA.cpp b/Engine/Plugins/Experimental/AVCodecs/NVCodecs/Source/NVENC/Private/Video/Encoders/VideoEncoderNVENCCUDA.cpp
new file mode 100644
--- /dev/null
+++ b/Engine/Plugins/Experimental/AVCodecs/NVCodecs/Source/NVENC/Private/Video/Encoders/VideoEncoderNVENCCUDA.cpp
@@ -0,0 +1,218 @@
+// Copyright Epic Games, Inc. All Rights Reserved.
+
+#include "Video/Encoders/VideoEncoderNVENC.h"
+#include "Templates/AlignmentTemplates.h"
+
+FVideoEncoderNVENCCUDA::~FVideoEncoderNVENCCUDA()
+{
+	Close();
+}
+
+FAVResult FVideoEncoderNVENCCUDA::Open(TSharedRef<FAVDevice> const& NewDevice, TSharedRef<FAVInstance> const& NewInstance)
+{
+	Close();
+	
+	TVideoEncoder::Open(NewDevice, NewInstance);
+
+	NV_ENC_STRUCT(NV_ENC_OPEN_ENCODE_SESSION_EX_PARAMS, SessionParams);
+	SessionParams.apiVersion = NVENCAPI_VERSION;
+	SessionParams.deviceType = NV_ENC_DEVICE_TYPE_CUDA;
+	SessionParams.device = GetDevice()->GetContext<FVideoContextCUDA>()->Raw;
+
+	return FEncoderNVENC::Open(NewDevice, NewInstance, SessionParams);
+}
+
+void FVideoEncoderNVENCCUDA::Close()
+{
+	if (IsOpen())
+	{
+		if (Encoder != nullptr)
+		{
+			NVENCSTATUS const Result = FAPI::Get<FNVENC>().nvEncDestroyEncoder(Encoder);
+			if (Result != NV_ENC_SUCCESS)
+			{
+				FAVResult::Log(EAVResult::ErrorDestroying, TEXT("Failed to destroy encoder"), TEXT("NVENC-CUDA"), Result);
+			}
+
+			Encoder = nullptr;
+		}
+
+		if (Buffer != nullptr)
+		{
+			NVENCSTATUS const Result = FAPI::Get<FNVENC>().nvEncDestroyBitstreamBuffer(Encoder, Buffer);
+			if (Result != NV_ENC_SUCCESS)
+			{
+				FAVResult::Log(EAVResult::ErrorDestroying, TEXT("Failed to destroy output buffer"), TEXT("NVENC-CUDA"), Result);
+			}
+
+			Buffer = nullptr;
+		}
+	}
+}
+
+bool FVideoEncoderNVENCCUDA::IsOpen() const
+{
+	return Encoder != nullptr;
+}
+
+FAVResult FVideoEncoderNVENCCUDA::ApplyConfig()
+{
+	FAVResult Result = FEncoderNVENC::ApplyConfig(AppliedConfig, EditPendingConfig());
+	if (Result.IsNotSuccess())
+	{
+		return Result;
+	}
+
+	return TVideoEncoder::ApplyConfig();
+}
+
+bool FVideoEncoderNVENCCUDA::IsInitialized() const
+{
+	return IsOpen() && Buffer != nullptr;
+}
+
+FAVResult FVideoEncoderNVENCCUDA::SendFrame(TSharedPtr<FVideoResourceCUDA> const& Resource, uint32 Timestamp, bool bForceKeyframe)
+{
+	if (!IsOpen())
+	{
+		return FAVResult(EAVResult::ErrorInvalidState, TEXT("Encoder not open"), TEXT("NVENC-CUDA"));
+	}
+
+	FAVResult AVResult = ApplyConfig();
+	if (AVResult.IsNotSuccess())
+	{
+		return AVResult;
+	}
+
+	NV_ENC_STRUCT(NV_ENC_REGISTER_RESOURCE, InputRegisterResource);
+	NV_ENC_STRUCT(NV_ENC_MAP_INPUT_RESOURCE, InputMapResource);
+	NV_ENC_STRUCT(NV_ENC_PIC_PARAMS, Picture);
+
+	NV_ENC_STRUCT(NV_ENC_LOCK_BITSTREAM, BitstreamLock);
+	BitstreamLock.outputBitstream = Buffer;
+
+	if (Resource.IsValid())
+	{
+		TAVResult<NV_ENC_BUFFER_FORMAT> const ConvertedPixelFormat = FVideoEncoderConfigNVENC::ConvertFormat(Resource->GetFormat());
+		if (ConvertedPixelFormat.IsNotSuccess())
+		{
+			return ConvertedPixelFormat;
+		}
+
+		DebugDumpFrame(Resource, Timestamp);
+
+		AVResult = RegisterResource(
+			InputRegisterResource,
+			NV_ENC_INPUT_RESOURCE_TYPE_CUDAARRAY, // ResourceType
+			ConvertedPixelFormat,				  // ResourceFormat
+			Resource->GetRaw(),					  // Resource
+			Resource->GetWidth(),				  // Width
+			Resource->GetHeight(),				  // Height
+			Resource->GetStride(),				  // Stride
+			NV_ENC_INPUT_IMAGE					  // BufferUsage
+		);
+		if (AVResult.IsNotSuccess())
+		{
+			return AVResult;
+		}
+
+		AVResult = MapResource(InputMapResource, InputRegisterResource.registeredResource);
+		if (AVResult.IsNotSuccess())
+		{
+			return AVResult;
+		}
+
+		Picture.inputWidth = InputRegisterResource.width;
+		Picture.inputHeight = InputRegisterResource.height;
+		Picture.inputPitch = InputRegisterResource.pitch;
+		Picture.inputBuffer = InputMapResource.mappedResource;
+		Picture.bufferFmt = InputMapResource.mappedBufferFmt;
+		Picture.pictureStruct = NV_ENC_PIC_STRUCT_FRAME;
+		Picture.inputTimeStamp = Timestamp;
+		Picture.outputBitstream = Buffer;
+
+		if (bForceKeyframe)
+		{
+			Picture.encodePicFlags |= NV_ENC_PIC_FLAG_FORCEIDR;
+		}
+	}
+	else
+	{
+		Picture.encodePicFlags = NV_ENC_PIC_FLAG_EOS;
+	}
+
+	if (Resource.IsValid())
+	{
+		Resource->Lock();
+	}
+
+	AVResult = FEncoderNVENC::SendFrame(Picture, BitstreamLock);
+
+	if (Resource.IsValid())
+	{
+		Resource->Unlock();
+	}
+
+	if (AVResult.IsNotSuccess())
+	{
+		// Here we just log the error instead of returning so that we ensure we unmap and unregister the input resources
+		AVResult.Log();
+	}
+
+	if (FAVResult UnmapResult = UnmapResource(InputMapResource.mappedResource); UnmapResult.IsNotSuccess())
+	{
+		UnmapResult.Log();
+		if(UnmapResult < AVResult)
+		{
+			AVResult = UnmapResult;
+		}
+	}
+
+	if (FAVResult UnregisterResult = UnregisterResource(InputRegisterResource.registeredResource); UnregisterResult.IsNotSuccess())
+	{
+		UnregisterResult.Log();
+		if (UnregisterResult < AVResult)
+		{
+			AVResult = UnregisterResult;
+		}
+	}
+
+	return AVResult;
+}
+
+FAVResult FVideoEncoderNVENCCUDA::ReceivePacket(FVideoPacket& OutPacket)
+{
+	return FEncoderNVENC::ReceivePacket(OutPacket);
+}
+
+FAVResult FVideoEncoderNVENCCUDA::OnPostInitialization(FVideoEncoderConfigNVENC& PendingConfig)
+{
+	NV_ENC_STRUCT(NV_ENC_CREATE_BITSTREAM_BUFFER, CreateBuffer);
+
+	NVENCSTATUS const Result = FAPI::Get<FNVENC>().nvEncCreateBitstreamBuffer(Encoder, &CreateBuffer);
+	if (Result != NV_ENC_SUCCESS)
+	{
+		return FAVResult(EAVResult::ErrorCreating, TEXT("Failed to create output buffer"), TEXT("NVENC-CUDA"), Result);
+	}
+
+	Buffer = CreateBuffer.bitstreamBuffer;
+
+	return FAVResult(EAVResult::Success, TEXT("Created output bitstream buffer."), TEXT("NVENC-CUDA"));
+}
+
+void FVideoEncoderNVENCCUDA::DebugDumpFrame(TSharedPtr<FVideoResourceCUDA> const& Resource, uint32 Timestamp)
+{
+	// START DEBUG
+#if DEBUG_DUMP_TO_DISK
+	{
+		TArray64<uint8> OutData;
+
+		Resource->ReadData(OutData);
+
+		FString SaveName = FString::Printf(TEXT("%s/DumpInput/image%05d.%s"), *FPaths::ProjectSavedDir(), Timestamp, ConvertedPixelFormat == NV_ENC_BUFFER_FORMAT_NV12 ? TEXT("nv12") : TEXT("p016"));
+
+		FFileHelper::SaveArrayToFile(OutData, *SaveName);
+	}
+#endif
+	// END DEBUG
+}
\ No newline at end of file
diff --git a/Engine/Plugins/Experimental/AVCodecs/NVCodecs/Source/NVENC/Private/Video/Encoders/VideoEncoderNVENCD3D11.cpp b/Engine/Plugins/Experimental/AVCodecs/NVCodecs/Source/NVENC/Private/Video/Encoders/VideoEncoderNVENCD3D11.cpp
new file mode 100644
--- /dev/null
+++ b/Engine/Plugins/Experimental/AVCodecs/NVCodecs/Source/NVENC/Private/Video/Encoders/VideoEncoderNVENCD3D11.cpp
@@ -0,0 +1,265 @@
+// Copyright Epic Games, Inc. All Rights Reserved.
+
+#include "Video/Encoders/VideoEncoderNVENC.h"
+#include "Templates/AlignmentTemplates.h"
+
+#if PLATFORM_WINDOWS
+/**
+ * You may be asking yourself, why on Earth would we need to create a seperate D3D11 device when using NVENC?
+ * Well, it basically comes down to the D3D11 device created by Unreal Engine is not compatible with NVENC due to sdk version changes (we think?)
+ * and Unreal Engine is in no hurry to bump D3D11 version - so we create our own device. This works only because
+ * the D3D11 textures we are using are created as shared resources and accessed through shared handles.
+ * Without this chicainery the D3D11 device will be ejected and Unreal Engine and NVENC will crash.
+ */
+FAVResult CreateD3D11Device(TSharedRef<FAVDevice> const& InDevice, TRefCountPtr<ID3D11Device>& OutEncoderDevice, TRefCountPtr<ID3D11DeviceContext>& OutEncoderDeviceContext)
+{
+	TRefCountPtr<IDXGIDevice>  DXGIDevice;
+	TRefCountPtr<IDXGIAdapter> Adapter;
+
+	HRESULT Result = InDevice->GetContext<FVideoContextD3D11>()->Device->QueryInterface(__uuidof(IDXGIDevice), (void**)DXGIDevice.GetInitReference());
+	if (Result != S_OK)
+	{
+		return FAVResult(EAVResult::Fatal, FString::Printf(TEXT("ID3D11Device::QueryInterface() failed 0x%X - %s."), (uint32)Result, *AVUtils::AVGetComErrorDescription(Result)), TEXT("NVENC-D3D11"));
+	}
+	else if ((Result = DXGIDevice->GetAdapter(Adapter.GetInitReference())) != S_OK)
+	{
+		return FAVResult(EAVResult::Fatal, FString::Printf(TEXT("DXGIDevice::GetAdapter() failed 0x%X - %s."), (uint32)Result, *AVUtils::AVGetComErrorDescription(Result)), TEXT("NVENC-D3D11"));
+	}
+
+	uint32			  DeviceFlags = 0;
+	D3D_FEATURE_LEVEL FeatureLevel = D3D_FEATURE_LEVEL_11_0;
+	D3D_FEATURE_LEVEL ActualFeatureLevel;
+
+	if ((Result = D3D11CreateDevice(
+			 Adapter,
+			 D3D_DRIVER_TYPE_UNKNOWN,
+			 NULL,
+			 DeviceFlags,
+			 &FeatureLevel,
+			 1,
+			 D3D11_SDK_VERSION,
+			 OutEncoderDevice.GetInitReference(),
+			 &ActualFeatureLevel,
+			 OutEncoderDeviceContext.GetInitReference()))
+		!= S_OK)
+	{
+		return FAVResult(EAVResult::Fatal, FString::Printf(TEXT("D3D11CreateDevice() failed 0x%X - %s."), (uint32)Result, *AVUtils::AVGetComErrorDescription(Result)), TEXT("NVENC-D3D11"));
+	}
+
+	return FAVResult(EAVResult::Success, TEXT("Created D3D11 device for NVENC."), TEXT("NVENC-D3D11"));
+}
+
+FVideoEncoderNVENCD3D11::~FVideoEncoderNVENCD3D11()
+{
+	Close();
+}
+
+FAVResult FVideoEncoderNVENCD3D11::Open(TSharedRef<FAVDevice> const& NewDevice, TSharedRef<FAVInstance> const& NewInstance)
+{
+	Close();
+
+	FAVResult Result = CreateD3D11Device(NewDevice, EncoderDevice, EncoderDeviceContext);
+
+	// If we failed to create device early out.
+	if (Result.IsNotSuccess())
+	{
+		return Result;
+	}
+
+	TVideoEncoder::Open(NewDevice, NewInstance);
+
+	NV_ENC_STRUCT(NV_ENC_OPEN_ENCODE_SESSION_EX_PARAMS, SessionParams);
+	SessionParams.apiVersion = NVENCAPI_VERSION;
+	SessionParams.deviceType = NV_ENC_DEVICE_TYPE_DIRECTX;
+	SessionParams.device = EncoderDevice;
+	// SessionParams.device = GetDevice()->GetContext<FVideoContextD3D11>()->Device;
+
+	return FEncoderNVENC::Open(NewDevice, NewInstance, SessionParams);
+}
+
+void FVideoEncoderNVENCD3D11::Close()
+{
+	if (IsOpen())
+	{
+		if (Encoder != nullptr)
+		{
+			NVENCSTATUS const Result = FAPI::Get<FNVENC>().nvEncDestroyEncoder(Encoder);
+			if (Result != NV_ENC_SUCCESS)
+			{
+				FAVResult::Log(EAVResult::ErrorDestroying, TEXT("Failed to destroy encoder"), TEXT("NVENC-D3D11"), Result);
+			}
+
+			Encoder = nullptr;
+		}
+
+		if (Buffer != nullptr)
+		{
+			NVENCSTATUS const Result = FAPI::Get<FNVENC>().nvEncDestroyBitstreamBuffer(Encoder, Buffer);
+			if (Result != NV_ENC_SUCCESS)
+			{
+				FAVResult::Log(EAVResult::ErrorDestroying, TEXT("Failed to destroy output buffer"), TEXT("NVENC-D3D11"), Result);
+			}
+
+			Buffer = nullptr;
+		}
+	}
+}
+
+bool FVideoEncoderNVENCD3D11::IsOpen() const
+{
+	return Encoder != nullptr;
+}
+
+bool FVideoEncoderNVENCD3D11::IsInitialized() const
+{
+	return IsOpen() && Buffer != nullptr;
+}
+
+FAVResult FVideoEncoderNVENCD3D11::ApplyConfig()
+{
+	FAVResult Result = FEncoderNVENC::ApplyConfig(AppliedConfig, EditPendingConfig());
+	if (Result.IsNotSuccess())
+	{
+		return Result;
+	}
+
+	return TVideoEncoder::ApplyConfig();
+}
+
+FAVResult FVideoEncoderNVENCD3D11::SendFrame(TSharedPtr<FVideoResourceD3D11> const& Resource, uint32 Timestamp, bool bForceKeyframe)
+{
+	if (!IsOpen())
+	{
+		return FAVResult(EAVResult::ErrorInvalidState, TEXT("Encoder not open"), TEXT("NVENC-D3D11"));
+	}
+
+	FAVResult AVResult = ApplyConfig();
+	if (AVResult.IsNotSuccess())
+	{
+		return AVResult;
+	}
+
+	TRefCountPtr<ID3D11Texture2D>& Tex = const_cast<TRefCountPtr<ID3D11Texture2D>&>(Resource->GetRaw());
+
+	// Very important: We create a new device for NVENC - we must share out D3D11 texture with the new device by "Opening" it to that device.
+	HRESULT HResult = EncoderDevice->OpenSharedResource(Resource->GetSharedHandle(), __uuidof(ID3D11Texture2D), (void**)(Tex.GetInitReference()));
+	if (HResult != S_OK)
+	{
+		return FAVResult(EAVResult::Fatal, TEXT("Failed to open shared handle."), TEXT("NVENC-D3D11"), HResult);
+	}
+
+	NV_ENC_STRUCT(NV_ENC_REGISTER_RESOURCE, InputRegisterResource);
+	NV_ENC_STRUCT(NV_ENC_MAP_INPUT_RESOURCE, InputMapResource);
+	NV_ENC_STRUCT(NV_ENC_PIC_PARAMS, Picture);
+
+	NV_ENC_STRUCT(NV_ENC_LOCK_BITSTREAM, BitstreamLock);
+	BitstreamLock.outputBitstream = Buffer;
+
+	if (Resource.IsValid())
+	{
+		TAVResult<NV_ENC_BUFFER_FORMAT> const ConvertedPixelFormat = FVideoEncoderConfigNVENC::ConvertFormat(Resource->GetFormat());
+		if (ConvertedPixelFormat.IsNotSuccess())
+		{
+			return ConvertedPixelFormat;
+		}
+
+		AVResult = RegisterResource(
+			InputRegisterResource,
+			NV_ENC_INPUT_RESOURCE_TYPE_DIRECTX, // ResourceType
+			ConvertedPixelFormat,				// ResourceFormat
+			Resource->GetRaw(),					// Resource
+			Resource->GetWidth(),				// Width
+			Resource->GetHeight(),				// Height
+			Resource->GetStride(),				// Stride
+			NV_ENC_INPUT_IMAGE					// BufferUsage
+		);
+		if (AVResult.IsNotSuccess())
+		{
+			return AVResult;
+		}
+
+		AVResult = MapResource(InputMapResource, InputRegisterResource.registeredResource);
+        if (AVResult.IsNotSuccess())
+		{
+			return AVResult;
+		}
+
+		Picture.inputWidth = InputRegisterResource.width;
+		Picture.inputHeight = InputRegisterResource.height;
+		Picture.inputPitch = InputRegisterResource.pitch;
+		Picture.inputBuffer = InputMapResource.mappedResource;
+		Picture.bufferFmt = InputMapResource.mappedBufferFmt;
+		Picture.pictureStruct = NV_ENC_PIC_STRUCT_FRAME;
+		Picture.inputTimeStamp = Timestamp;
+		Picture.outputBitstream = Buffer;
+
+		if (bForceKeyframe)
+		{
+			Picture.encodePicFlags |= NV_ENC_PIC_FLAG_FORCEIDR;
+		}
+	}
+	else
+	{
+		Picture.encodePicFlags = NV_ENC_PIC_FLAG_EOS;
+	}
+
+	if (Resource.IsValid())
+	{
+		Resource->Lock();
+	}
+
+	AVResult = FEncoderNVENC::SendFrame(Picture, BitstreamLock);
+
+	if (Resource.IsValid())
+	{
+		Resource->Unlock();
+	}
+
+	if (AVResult.IsNotSuccess())
+	{
+		// Here we just log the error instead of returning so that we ensure we unmap and unregister the input resources
+		AVResult.Log();
+	}
+
+	if (FAVResult UnmapResult = UnmapResource(InputMapResource.mappedResource); UnmapResult.IsNotSuccess())
+	{
+		UnmapResult.Log();
+		if(UnmapResult < AVResult)
+		{
+			AVResult = UnmapResult;
+		}
+	}
+
+	if (FAVResult UnregisterResult = UnregisterResource(InputRegisterResource.registeredResource); UnregisterResult.IsNotSuccess())
+	{
+		UnregisterResult.Log();
+		if (UnregisterResult < AVResult)
+		{
+			AVResult = UnregisterResult;
+		}
+	}
+
+	return AVResult;
+}
+
+FAVResult FVideoEncoderNVENCD3D11::ReceivePacket(FVideoPacket& OutPacket)
+{
+	return FEncoderNVENC::ReceivePacket(OutPacket);
+}
+
+FAVResult FVideoEncoderNVENCD3D11::OnPostInitialization(FVideoEncoderConfigNVENC& PendingConfig)
+{
+	NV_ENC_STRUCT(NV_ENC_CREATE_BITSTREAM_BUFFER, CreateBuffer);
+
+	NVENCSTATUS const Result = FAPI::Get<FNVENC>().nvEncCreateBitstreamBuffer(Encoder, &CreateBuffer);
+	if (Result != NV_ENC_SUCCESS)
+	{
+		return FAVResult(EAVResult::ErrorCreating, TEXT("Failed to create output buffer"), TEXT("NVENC-D3D11"), Result);
+	}
+
+	Buffer = CreateBuffer.bitstreamBuffer;
+
+	return FAVResult(EAVResult::Success, TEXT("Created output bitstream buffer."), TEXT("NVENC-D3D11"));
+}
+
+#endif
\ No newline at end of file
diff --git a/Engine/Plugins/Experimental/AVCodecs/NVCodecs/Source/NVENC/Private/Video/Encoders/VideoEncoderNVENCD3D12.cpp b/Engine/Plugins/Experimental/AVCodecs/NVCodecs/Source/NVENC/Private/Video/Encoders/VideoEncoderNVENCD3D12.cpp
new file mode 100644
--- /dev/null
+++ b/Engine/Plugins/Experimental/AVCodecs/NVCodecs/Source/NVENC/Private/Video/Encoders/VideoEncoderNVENCD3D12.cpp
@@ -0,0 +1,306 @@
+// Copyright Epic Games, Inc. All Rights Reserved.
+
+#include "Video/Encoders/VideoEncoderNVENC.h"
+#include "Templates/AlignmentTemplates.h"
+
+#if PLATFORM_WINDOWS
+FAVResult CreateD3D12Fence(TSharedRef<FAVDevice> const& InDevice, TRefCountPtr<ID3D12Fence>& OutFence)
+{
+	TRefCountPtr<ID3D12Device> D3D12Device = InDevice->GetContext<FVideoContextD3D12>()->Device;
+
+	HRESULT Result;
+	if ((Result = D3D12Device->CreateFence(0, D3D12_FENCE_FLAG_NONE, IID_PPV_ARGS(OutFence.GetInitReference()))) != S_OK)
+	{
+		return FAVResult(EAVResult::Fatal, FString::Printf(TEXT("CreateD3D12Fence() failed 0x%X - %s."), (uint32)Result, *AVUtils::AVGetComErrorDescription(Result)), TEXT("NVENC-D3D12"));
+	}
+
+	return FAVResult(EAVResult::Success, TEXT("Created additional D3D12 resources for NVENC."), TEXT("NVENC-D3D12"));
+}
+
+FVideoEncoderNVENCD3D12::~FVideoEncoderNVENCD3D12()
+{
+	Close();
+}
+
+FAVResult FVideoEncoderNVENCD3D12::Open(TSharedRef<FAVDevice> const& NewDevice, TSharedRef<FAVInstance> const& NewInstance)
+{
+	Close();
+
+	FAVResult Result = CreateD3D12Fence(NewDevice, InputFence);
+	if (Result.IsNotSuccess())
+	{
+		return Result;
+	}
+
+	Result = CreateD3D12Fence(NewDevice, OutputFence);
+	if (Result.IsNotSuccess())
+	{
+		return Result;
+	}
+
+	TVideoEncoder::Open(NewDevice, NewInstance);
+
+	NV_ENC_STRUCT(NV_ENC_OPEN_ENCODE_SESSION_EX_PARAMS, SessionParams);
+	SessionParams.apiVersion = NVENCAPI_VERSION;
+	SessionParams.deviceType = NV_ENC_DEVICE_TYPE_DIRECTX;
+	SessionParams.device = NewDevice->GetContext<FVideoContextD3D12>()->Device;
+
+	return FEncoderNVENC::Open(NewDevice, NewInstance, SessionParams);
+}
+
+void FVideoEncoderNVENCD3D12::Close()
+{
+	if (IsOpen())
+	{
+		if (Encoder != nullptr)
+		{
+			NVENCSTATUS const Result = FAPI::Get<FNVENC>().nvEncDestroyEncoder(Encoder);
+			if (Result != NV_ENC_SUCCESS)
+			{
+				FAVResult::Log(EAVResult::ErrorDestroying, TEXT("Failed to destroy encoder"), TEXT("NVENC-D3D11"), Result);
+			}
+
+			Encoder = nullptr;
+		}
+	}
+}
+
+bool FVideoEncoderNVENCD3D12::IsOpen() const
+{
+	return Encoder != nullptr;
+}
+
+bool FVideoEncoderNVENCD3D12::IsInitialized() const
+{
+	return IsOpen() && OutputBitstreamResource.IsValid();
+}
+
+FAVResult FVideoEncoderNVENCD3D12::ApplyConfig()
+{
+	FAVResult Result = FEncoderNVENC::ApplyConfig(AppliedConfig, EditPendingConfig());
+	if (Result.IsNotSuccess())
+	{
+		return Result;
+	}
+
+	return TVideoEncoder::ApplyConfig();
+}
+
+FAVResult FVideoEncoderNVENCD3D12::SendFrame(TSharedPtr<FVideoResourceD3D12> const& Resource, uint32 Timestamp, bool bForceKeyframe)
+{
+	if (!IsOpen())
+	{
+		return FAVResult(EAVResult::ErrorInvalidState, TEXT("Encoder not open"), TEXT("NVENC-D3D12"));
+	}
+
+	{
+		FVideoEncoderConfigNVENC& PendingConfig = EditPendingConfig();
+
+		TAVResult<NV_ENC_BUFFER_FORMAT> const ConvertedPixelFormat = FVideoEncoderConfigNVENC::ConvertFormat(Resource->GetFormat());
+		if (ConvertedPixelFormat.IsNotSuccess())
+		{
+			return ConvertedPixelFormat;
+		}
+
+		if (ConvertedPixelFormat != NV_ENC_BUFFER_FORMAT_ARGB)
+		{
+			return FAVResult(EAVResult::ErrorUnsupported, TEXT("DX12 resources must be of the format DXGI_FORMAT_B8G8R8A8_UNORM"), TEXT("NVENC"));
+		}
+
+		PendingConfig.bufferFormat = ConvertedPixelFormat;
+
+		FAVResult AVResult = ApplyConfig();
+		if (AVResult.IsNotSuccess())
+		{
+			return AVResult;
+		}
+	}
+
+	FAVResult	AVResult = EAVResult::Success;
+
+	NV_ENC_STRUCT(NV_ENC_REGISTER_RESOURCE, InputRegisterResource);
+	NV_ENC_STRUCT(NV_ENC_MAP_INPUT_RESOURCE, InputMapResource);
+	NV_ENC_STRUCT(NV_ENC_INPUT_RESOURCE_D3D12, InputResource);
+	// Register and map input resource
+	if (Resource.IsValid())
+	{
+		AVResult = RegisterResource(
+			InputRegisterResource,
+			NV_ENC_INPUT_RESOURCE_TYPE_DIRECTX, // ResourceType
+			NV_ENC_BUFFER_FORMAT_ARGB,			// ResourceFormat
+			Resource->GetResource(),			// Resource
+			Resource->GetWidth(),				// Width
+			Resource->GetHeight(),				// Height
+			0,									// Stride
+			NV_ENC_INPUT_IMAGE					// BufferUsage
+		);
+		if (AVResult.IsNotSuccess())
+		{
+			return AVResult;
+		}
+
+		AVResult = MapResource(InputMapResource, InputRegisterResource.registeredResource);
+		if (AVResult.IsNotSuccess())
+		{
+			return AVResult;
+		}
+
+		InputResource.inputFencePoint.pFence = InputFence.GetReference();
+		InputResource.pInputBuffer = InputMapResource.mappedResource;
+		InputResource.inputFencePoint.waitValue = InputFenceVal;
+		InputResource.inputFencePoint.bWait = true;
+	}
+
+	FPlatformAtomics::InterlockedIncrement(&OutputFenceVal);
+	// Register and map output resource
+	NV_ENC_STRUCT(NV_ENC_REGISTER_RESOURCE, OutputRegisterResource);
+	NV_ENC_STRUCT(NV_ENC_MAP_INPUT_RESOURCE, OutputMapResource);
+	NV_ENC_STRUCT(NV_ENC_OUTPUT_RESOURCE_D3D12, OutputResource);
+	{
+		AVResult = RegisterResource(
+			OutputRegisterResource,
+			NV_ENC_INPUT_RESOURCE_TYPE_DIRECTX,										  // ResourceType
+			NV_ENC_BUFFER_FORMAT_U8,												  // ResourceFormat
+			OutputBitstreamResource.GetReference(),									  // Resource
+			Align(AppliedConfig.encodeWidth * AppliedConfig.encodeHeight * 4 * 2, 4), // Width
+			1,																		  // Height
+			0,																		  // Stride
+			NV_ENC_OUTPUT_BITSTREAM													  // BufferUsage
+		);
+		if (AVResult.IsNotSuccess())
+		{
+			return AVResult;
+		}
+
+		AVResult = MapResource(OutputMapResource, OutputRegisterResource.registeredResource);
+		if (AVResult.IsNotSuccess())
+		{
+			return AVResult;
+		}
+
+		OutputResource.outputFencePoint.pFence = OutputFence.GetReference();
+		OutputResource.pOutputBuffer = OutputMapResource.mappedResource;
+		OutputResource.outputFencePoint.signalValue = OutputFenceVal;
+		OutputResource.outputFencePoint.bSignal = true;
+	}
+
+	NV_ENC_STRUCT(NV_ENC_PIC_PARAMS, Picture);
+	{
+		Picture.pictureStruct = NV_ENC_PIC_STRUCT_FRAME;
+		Picture.inputTimeStamp = Timestamp;
+		Picture.inputBuffer = &InputResource;
+		Picture.bufferFmt = InputMapResource.mappedBufferFmt;
+		Picture.inputWidth = InputRegisterResource.width;
+		Picture.inputHeight = InputRegisterResource.height;
+		Picture.outputBitstream = &OutputResource;
+		Picture.completionEvent = nullptr;
+	}
+
+	if (!Resource.IsValid())
+	{
+		Picture.encodePicFlags = NV_ENC_PIC_FLAG_EOS;
+	}
+	else if (bForceKeyframe)
+	{
+		Picture.encodePicFlags |= NV_ENC_PIC_FLAG_FORCEIDR;
+	}
+
+	if (Resource.IsValid())
+	{
+		Resource->Lock();
+	}
+
+	NV_ENC_STRUCT(NV_ENC_LOCK_BITSTREAM, BitstreamLock);
+	BitstreamLock.outputBitstream = &OutputResource;
+
+	AVResult = FEncoderNVENC::SendFrame(Picture, BitstreamLock);
+
+	if (Resource.IsValid())
+	{
+		Resource->Unlock();
+	}
+
+	if (AVResult.IsNotSuccess())
+	{
+		// Here we just log the error instead of returning so that we ensure we unmap and unregister the input resources
+		AVResult.Log();
+	}
+
+	// Unmap and unregister input resources
+	if (FAVResult UnmapResult = UnmapResource(InputMapResource.mappedResource); UnmapResult.IsNotSuccess())
+	{
+		UnmapResult.Log();
+		if(UnmapResult < AVResult)
+		{
+			AVResult = UnmapResult;
+		}
+	}
+
+	if (FAVResult UnregisterResult = UnregisterResource(InputRegisterResource.registeredResource); UnregisterResult.IsNotSuccess())
+	{
+		UnregisterResult.Log();
+		if (UnregisterResult < AVResult)
+		{
+			AVResult = UnregisterResult;
+		}
+	}
+
+	// Unmap and unregister output resources
+	if (FAVResult UnmapResult = UnmapResource(OutputMapResource.mappedResource); UnmapResult.IsNotSuccess())
+	{
+		UnmapResult.Log();
+		if(UnmapResult < AVResult)
+		{
+			AVResult = UnmapResult;
+		}
+	}
+
+	if (FAVResult UnregisterResult = UnregisterResource(OutputRegisterResource.registeredResource); UnregisterResult.IsNotSuccess())
+	{
+		UnregisterResult.Log();
+		if (UnregisterResult < AVResult)
+		{
+			AVResult = UnregisterResult;
+		}
+	}
+
+	return AVResult;
+}
+
+FAVResult FVideoEncoderNVENCD3D12::ReceivePacket(FVideoPacket& OutPacket)
+{
+	return FEncoderNVENC::ReceivePacket(OutPacket);
+}
+
+FAVResult FVideoEncoderNVENCD3D12::OnPostInitialization(FVideoEncoderConfigNVENC& PendingConfig)
+{
+	TRefCountPtr<ID3D12Device> D3D12Device = GetDevice()->GetContext<FVideoContextD3D12>()->Device;
+
+	D3D12_HEAP_PROPERTIES HeapProps = {};
+	HeapProps.Type = D3D12_HEAP_TYPE_READBACK;
+	HeapProps.CPUPageProperty = D3D12_CPU_PAGE_PROPERTY_UNKNOWN;
+	HeapProps.MemoryPoolPreference = D3D12_MEMORY_POOL_UNKNOWN;
+
+	D3D12_RESOURCE_DESC ResourceDesc = {};
+	ResourceDesc.Dimension = D3D12_RESOURCE_DIMENSION_BUFFER;
+	ResourceDesc.Alignment = 0;
+	ResourceDesc.Width = Align(PendingConfig.encodeWidth * PendingConfig.encodeHeight * 4 * 2, 4);
+	ResourceDesc.Height = 1;
+	ResourceDesc.DepthOrArraySize = 1;
+	ResourceDesc.MipLevels = 1;
+	ResourceDesc.Format = DXGI_FORMAT_UNKNOWN;
+	ResourceDesc.SampleDesc.Count = 1;
+	ResourceDesc.SampleDesc.Quality = 0;
+	ResourceDesc.Layout = D3D12_TEXTURE_LAYOUT_ROW_MAJOR;
+	ResourceDesc.Flags = D3D12_RESOURCE_FLAG_NONE;
+
+	HRESULT Result;
+	if ((Result = D3D12Device->CreateCommittedResource(&HeapProps, D3D12_HEAP_FLAG_NONE, &ResourceDesc, D3D12_RESOURCE_STATE_COPY_DEST, nullptr, IID_PPV_ARGS(OutputBitstreamResource.GetInitReference()))) != S_OK)
+	{
+		return FAVResult(EAVResult::Fatal, FString::Printf(TEXT("CreateCommittedResource() failed 0x%X - %s."), (uint32)Result, *AVUtils::AVGetComErrorDescription(Result)), TEXT("NVENC-D3D12"));
+	}
+
+	return FAVResult(EAVResult::Success, TEXT("Created additional D3D12 resources for NVENC."), TEXT("NVENC-D3D12"));
+}
+
+#endif
\ No newline at end of file
diff --git a/Engine/Plugins/Experimental/AVCodecs/NVCodecs/Source/NVENC/Public/Video/Encoders/VideoEncoderNVENC.h b/Engine/Plugins/Experimental/AVCodecs/NVCodecs/Source/NVENC/Public/Video/Encoders/VideoEncoderNVENC.h
--- a/Engine/Plugins/Experimental/AVCodecs/NVCodecs/Source/NVENC/Public/Video/Encoders/VideoEncoderNVENC.h
+++ b/Engine/Plugins/Experimental/AVCodecs/NVCodecs/Source/NVENC/Public/Video/Encoders/VideoEncoderNVENC.h
@@ -7,6 +7,7 @@
 #include "Video/VideoResource.h"
 #include "NVENC.h"
 #include "Video/Encoders/Configs/VideoEncoderConfigNVENC.h"
+#include "Templates/AlignmentTemplates.h"
 #include "Templates/RefCounting.h"
 #include "Containers/Queue.h"
 
@@ -14,6 +15,13 @@
 	#include "Video/Resources/D3D/VideoResourceD3D.h"
 #endif // PLATFORM_WINDOWS
 
+#if PLATFORM_WINDOWS
+namespace AVUtils
+{
+	const FString AVGetComErrorDescription(HRESULT Res);
+}
+#endif
+
 /**
  * This class is the core of the NVENC encoder logic.
  * It is embedded inside NVCodecEncoders using composition instead of inheritance because templates + inheritance was bit messy here where we need concrete types.
@@ -21,142 +29,136 @@
 class NVENC_API FEncoderNVENC
 {
 private:
-	void* Encoder = nullptr;
-	NV_ENC_OUTPUT_PTR Buffer = nullptr;
-	TQueue<FVideoPacket> Packets;
-
-	NV_ENC_DEVICE_TYPE  SessionDeviceType;
-	void* SessionDevice;
-
 	uint32 MaxDeviceEncodeWidth = 0;
 	uint32 MaxDeviceEncodeHeight = 0;
 	bool bHasMaxDeviceResolution = false;
 
-	FAVResult ReOpen();
 	void SetMaxResolution(FVideoEncoderConfigNVENC& PendingConfig);
 	void GetMaxDeviceEncodeResolution(const FVideoEncoderConfigNVENC& PendingConfig);
 
 public:
-	// Begin matching the TVideoEncoder interface
-	virtual ~FEncoderNVENC();
-	virtual bool IsOpen() const;
-	virtual FAVResult Open(TSharedRef<FAVDevice> const& NewDevice, TSharedRef<FAVInstance> const& NewInstance, TFunction<void(NV_ENC_OPEN_ENCODE_SESSION_EX_PARAMS&)> SetupEncoderSessionFunc);
-	virtual void Close();
-	virtual FAVResult ApplyConfig(FVideoEncoderConfigNVENC const& AppliedConfig, FVideoEncoderConfigNVENC& PendingConfig, TFunction<FAVResult()> ApplyConfigFunc);
-#if PLATFORM_WINDOWS
-	virtual FAVResult CreateD3D11Device(TSharedRef<FAVDevice> const& InDevice, TRefCountPtr<ID3D11Device>& OutEncoderDevice, TRefCountPtr<ID3D11DeviceContext>& OutEncoderDeviceContext);
-	virtual FAVResult SendFrameD3D11( TRefCountPtr<ID3D11Device> Device, TSharedPtr<FVideoResourceD3D11> const& Resource, uint32 Timestamp, bool bForceKeyframe, TFunction<FAVResult()> ApplyConfigFunc);
-#endif // PLATFORM_WINDOWS
-	virtual FAVResult SendFrameCUDA(TSharedPtr<FVideoResourceCUDA> const& Resource, uint32 Timestamp, bool bForceKeyframe, TFunction<FAVResult()> ApplyConfigFunc);
-	virtual FAVResult SendFrame(TSharedPtr<FVideoResource> const& Resource, uint32 Timestamp, bool bForceKeyframe, TFunction<FAVResult()> ApplyConfigFunc, TFunction<void(NV_ENC_REGISTER_RESOURCE&)> SetResourceToRegisterFunc);
-	virtual FAVResult ReceivePacket(FVideoPacket& OutPacket);
-	// End matching the TVideoEncoder interface
+	FEncoderNVENC() = default;
+	virtual ~FEncoderNVENC() = default;
+
+protected:
+	NV_ENC_DEVICE_TYPE  SessionDeviceType;
+	void* SessionDevice;
+
+	void* Encoder = nullptr;
+	TQueue<FVideoPacket> Packets;
+
+	FAVResult ReOpen();
+	FAVResult Open(TSharedRef<FAVDevice> const& NewDevice, TSharedRef<FAVInstance> const& NewInstance, NV_ENC_OPEN_ENCODE_SESSION_EX_PARAMS& EncoderSessionParams);
+	FAVResult ApplyConfig(FVideoEncoderConfigNVENC const& AppliedConfig, FVideoEncoderConfigNVENC& PendingConfig);
+	FAVResult SendFrame(NV_ENC_PIC_PARAMS& Input, NV_ENC_LOCK_BITSTREAM& BitstreamLock);
+	FAVResult ReceivePacket(FVideoPacket& OutPacket);
+
+	//
+	virtual bool IsOpen() const = 0;
+	virtual void Close() = 0;
+	virtual	bool IsInitialized() const = 0;
+	virtual FAVResult OnPostInitialization(FVideoEncoderConfigNVENC& PendingConfig) { return EAVResult::Success; }
+	//
+
+	FAVResult RegisterResource(NV_ENC_REGISTER_RESOURCE& OutRegisterResource, NV_ENC_INPUT_RESOURCE_TYPE ResourceType, NV_ENC_BUFFER_FORMAT ResourceFormat, void* Resource, uint32 Width, uint32 Height, uint32 Stride, NV_ENC_BUFFER_USAGE BufferUsage);
+	FAVResult UnregisterResource(NV_ENC_REGISTERED_PTR RegisteredResource);
+	FAVResult MapResource(NV_ENC_MAP_INPUT_RESOURCE& OutMapResource, NV_ENC_REGISTERED_PTR RegisteredResource);
+	FAVResult UnmapResource(NV_ENC_INPUT_PTR MappedResource);
 
-	bool IsInitialized() const;
-	FAVResult SendFrame(NV_ENC_PIC_PARAMS Input);
 	int GetCapability(GUID EncodeGUID, NV_ENC_CAPS CapsToQuery) const;
 };
 
 /**
  *  An AVCodecs VideoEncoder that takes a FVideoResourceCUDA.
  **/
-class NVENC_API FVideoEncoderNVENCCUDA : public TVideoEncoder<FVideoResourceCUDA, FVideoEncoderConfigNVENC>
+class NVENC_API FVideoEncoderNVENCCUDA : public TVideoEncoder<FVideoResourceCUDA, FVideoEncoderConfigNVENC>, public FEncoderNVENC
 {
 private:
-	TUniquePtr<FEncoderNVENC> Base;
+	NV_ENC_OUTPUT_PTR Buffer = nullptr;
 
 public:
-	FVideoEncoderNVENCCUDA()
-		: Base(MakeUnique<FEncoderNVENC>()) {}
+	FVideoEncoderNVENCCUDA() = default;
+	virtual ~FVideoEncoderNVENCCUDA() override;
 
 	// Begin TVideoEncoder interface
-	virtual ~FVideoEncoderNVENCCUDA() = default;
-	virtual bool IsOpen() const override { return Base->IsOpen(); }
-	virtual void Close() override { Base->Close(); };
-	virtual FAVResult ReceivePacket(FVideoPacket& OutPacket) override { return Base->ReceivePacket(OutPacket); }
-
-	virtual FAVResult ApplyConfig() override
-	{
-		return Base->ApplyConfig(AppliedConfig, EditPendingConfig(), [this]() {
-			return TVideoEncoder::ApplyConfig();
-		});
-	}
-
-	virtual FAVResult Open(TSharedRef<FAVDevice> const& NewDevice, TSharedRef<FAVInstance> const& NewInstance) override
-	{
-		return Base->Open(NewDevice, NewInstance, [this, NewDevice, NewInstance](NV_ENC_OPEN_ENCODE_SESSION_EX_PARAMS& SessionParams) {
-			TVideoEncoder::Open(NewDevice, NewInstance);
-			SessionParams.deviceType = NV_ENC_DEVICE_TYPE_CUDA;
-			SessionParams.device = GetDevice()->GetContext<FVideoContextCUDA>()->Raw;
-		});
-	}
-
-	virtual FAVResult SendFrame(TSharedPtr<FVideoResourceCUDA> const& Resource, uint32 Timestamp, bool bForceKeyframe = false) override
-	{
-		return Base->SendFrameCUDA(Resource, Timestamp, bForceKeyframe, [this]() {
-			return ApplyConfig();
-		});
-	}
+	virtual FAVResult Open(TSharedRef<FAVDevice> const& NewDevice, TSharedRef<FAVInstance> const& NewInstance) override;
+	virtual void Close() override;
+	virtual bool IsOpen() const override;
+	virtual FAVResult ApplyConfig() override;
+	virtual FAVResult SendFrame(TSharedPtr<FVideoResourceCUDA> const& Resource, uint32 Timestamp, bool bForceKeyframe = false) override;
+	virtual FAVResult ReceivePacket(FVideoPacket& OutPacket) override;
 	// End TVideoEncoder interface
+
+	// Begin FEncoderNVENC interface
+	virtual	bool IsInitialized() const override;
+	virtual FAVResult OnPostInitialization(FVideoEncoderConfigNVENC& PendingConfig) override;
+	// End FEncoderNVENC interface
+
+private:
+	void DebugDumpFrame(TSharedPtr<FVideoResourceCUDA> const& Resource, uint32 Timestamp);
 };
 
 #if PLATFORM_WINDOWS
-
 /**
  *  An AVCodecs VideoEncoder that takes a FVideoResourceD3D11.
  **/
-class NVENC_API FVideoEncoderNVENCD3D11 : public TVideoEncoder<FVideoResourceD3D11, FVideoEncoderConfigNVENC>
+class NVENC_API FVideoEncoderNVENCD3D11 : public TVideoEncoder<FVideoResourceD3D11, FVideoEncoderConfigNVENC>, public FEncoderNVENC
 {
 private:
-	TUniquePtr<FEncoderNVENC> Base;
 	TRefCountPtr<ID3D11Device> EncoderDevice;
 	TRefCountPtr<ID3D11DeviceContext> EncoderDeviceContext;
 
+	NV_ENC_OUTPUT_PTR Buffer = nullptr;
+
 public:
-	FVideoEncoderNVENCD3D11()
-		: Base(MakeUnique<FEncoderNVENC>()) {}
+	FVideoEncoderNVENCD3D11() = default;
+	virtual ~FVideoEncoderNVENCD3D11() override;
 
 	// Begin TVideoEncoder interface
-	virtual ~FVideoEncoderNVENCD3D11() = default;
-	virtual bool IsOpen() const override { return Base->IsOpen(); }
-	virtual void Close() override { Base->Close(); };
-	virtual FAVResult ReceivePacket(FVideoPacket& OutPacket) override { return Base->ReceivePacket(OutPacket); }
-
-	virtual FAVResult ApplyConfig() override
-	{
-		return Base->ApplyConfig(AppliedConfig, EditPendingConfig(), [this]() {
-			return TVideoEncoder::ApplyConfig();
-		});
-	};
-
-	virtual FAVResult Open(TSharedRef<FAVDevice> const& NewDevice, TSharedRef<FAVInstance> const& NewInstance) override
-	{
-
-		// then pass device in SendFrameD3D11 and do OpenHandle
-
-		return Base->Open(NewDevice, NewInstance, [this, NewDevice, NewInstance](NV_ENC_OPEN_ENCODE_SESSION_EX_PARAMS& SessionParams) {
-			FAVResult Result = Base->CreateD3D11Device(NewDevice, EncoderDevice, EncoderDeviceContext);
-
-			// If we failed to create device early out.
-			if (Result.IsNotSuccess())
-			{
-				return Result;
-			}
-
-			TVideoEncoder::Open(NewDevice, NewInstance);
-			SessionParams.deviceType = NV_ENC_DEVICE_TYPE_DIRECTX;
-			SessionParams.device = EncoderDevice;
-			// SessionParams.device = GetDevice()->GetContext<FVideoContextD3D11>()->Device;
-			return Result;
-		});
-	}
-
-	virtual FAVResult SendFrame(TSharedPtr<FVideoResourceD3D11> const& Resource, uint32 Timestamp, bool bForceKeyframe = false)
-	{
-		return Base->SendFrameD3D11(EncoderDevice, Resource, Timestamp, bForceKeyframe, [this]() {
-			return ApplyConfig();
-		});
-	}
+	virtual FAVResult Open(TSharedRef<FAVDevice> const& NewDevice, TSharedRef<FAVInstance> const& NewInstance) override;
+	virtual void Close() override;
+	virtual bool IsOpen() const override;
+	virtual FAVResult ApplyConfig() override;
+	virtual FAVResult SendFrame(TSharedPtr<FVideoResourceD3D11> const& Resource, uint32 Timestamp, bool bForceKeyframe = false) override;
+	virtual FAVResult ReceivePacket(FVideoPacket& OutPacket) override;
 	// End TVideoEncoder interface
+
+	// Begin FEncoderNVENC interface
+	virtual	bool IsInitialized() const override;
+	virtual FAVResult OnPostInitialization(FVideoEncoderConfigNVENC& PendingConfig) override;
+	// End FEncoderNVENC interface
+};
+
+/**
+ *  An AVCodecs VideoEncoder that takes a FVideoResourceD3D12.
+ **/
+class NVENC_API FVideoEncoderNVENCD3D12 : public TVideoEncoder<FVideoResourceD3D12, FVideoEncoderConfigNVENC>, public FEncoderNVENC
+{
+private:
+	TRefCountPtr<ID3D12Fence> InputFence;
+	int64 InputFenceVal = 0;
+
+	TRefCountPtr<ID3D12Fence> OutputFence;
+	int64 OutputFenceVal = 0;
+
+	TRefCountPtr<ID3D12Resource> OutputBitstreamResource;
+
+public:
+	FVideoEncoderNVENCD3D12() = default;
+	virtual ~FVideoEncoderNVENCD3D12() override;
+
+	// Begin TVideoEncoder interface
+	virtual FAVResult Open(TSharedRef<FAVDevice> const& NewDevice, TSharedRef<FAVInstance> const& NewInstance) override;
+	virtual void Close() override;
+	virtual bool IsOpen() const override;
+	virtual FAVResult ApplyConfig() override;
+	virtual FAVResult SendFrame(TSharedPtr<FVideoResourceD3D12> const& Resource, uint32 Timestamp, bool bForceKeyframe = false) override;
+	virtual FAVResult ReceivePacket(FVideoPacket& OutPacket) override;
+	// End TVideoEncoder interface
+
+	// Begin FEncoderNVENC interface
+	virtual	bool IsInitialized() const override;
+	virtual FAVResult OnPostInitialization(FVideoEncoderConfigNVENC& PendingConfig) override;
+	// End FEncoderNVENC interface
 };
 #endif // PLATFORM_WINDOWS
\ No newline at end of file
