[AVCodecs] Fix: Address potential bugs and issues with destruction order

diff --git a/Engine/Plugins/Experimental/AVCodecs/NVCodecs/Source/NVENC/Private/NVENCModule.cpp b/Engine/Plugins/Experimental/AVCodecs/NVCodecs/Source/NVENC/Private/NVENCModule.cpp
--- a/Engine/Plugins/Experimental/AVCodecs/NVCodecs/Source/NVENC/Private/NVENCModule.cpp
+++ b/Engine/Plugins/Experimental/AVCodecs/NVCodecs/Source/NVENC/Private/NVENCModule.cpp
@@ -16,6 +16,15 @@ public:
 	virtual void StartupModule() override
 	{
 		TFunction<bool(TSharedRef<FAVDevice> const&, GUID)> CheckCodecSupport = [](TSharedRef<FAVDevice> const& Device, GUID CodecGUID) {
+			void* Encoder = nullptr;
+			ON_SCOPE_EXIT
+			{
+				if (Encoder)
+				{
+					FAPI::Get<FNVENC>().nvEncDestroyEncoder(Encoder);
+				}
+			};
+
 			if(!Device->HasContext<FVideoContextCUDA>())
 			{
 				return false;
@@ -31,7 +40,6 @@ public:
 			SessionParams.deviceType = NV_ENC_DEVICE_TYPE_CUDA;
 			SessionParams.device = Device->GetContext<FVideoContextCUDA>()->Raw;
 
-			void* Encoder = nullptr;
 			NVENCSTATUS Result = FAPI::Get<FNVENC>().nvEncOpenEncodeSessionEx(&const_cast<NV_ENC_OPEN_ENCODE_SESSION_EX_PARAMS&>(SessionParams), &Encoder);
 			if (Result != NV_ENC_SUCCESS)
 			{
@@ -63,12 +71,6 @@ public:
          		}
 			}
 
-			Result = FAPI::Get<FNVENC>().nvEncDestroyEncoder(Encoder);
-			if (Result != NV_ENC_SUCCESS)
-			{
-				return false;
-			}
-
 			return bSupported;
 		};
 
diff --git a/Engine/Plugins/Experimental/AVCodecs/NVCodecs/Source/NVENC/Private/Video/Encoders/VideoEncoderNVENCCUDA.cpp b/Engine/Plugins/Experimental/AVCodecs/NVCodecs/Source/NVENC/Private/Video/Encoders/VideoEncoderNVENCCUDA.cpp
--- a/Engine/Plugins/Experimental/AVCodecs/NVCodecs/Source/NVENC/Private/Video/Encoders/VideoEncoderNVENCCUDA.cpp
+++ b/Engine/Plugins/Experimental/AVCodecs/NVCodecs/Source/NVENC/Private/Video/Encoders/VideoEncoderNVENCCUDA.cpp
@@ -26,26 +26,26 @@ void FVideoEncoderNVENCCUDA::Close()
 {
 	if (IsOpen())
 	{
-		if (Encoder != nullptr)
+		if (Buffer != nullptr)
 		{
-			NVENCSTATUS const Result = FAPI::Get<FNVENC>().nvEncDestroyEncoder(Encoder);
+			NVENCSTATUS const Result = FAPI::Get<FNVENC>().nvEncDestroyBitstreamBuffer(Encoder, Buffer);
 			if (Result != NV_ENC_SUCCESS)
 			{
-				FAVResult::Log(EAVResult::ErrorDestroying, TEXT("Failed to destroy encoder"), TEXT("NVENC-CUDA"), Result);
+				FAVResult::Log(EAVResult::ErrorDestroying, TEXT("Failed to destroy output buffer"), TEXT("NVENC-CUDA"), Result);
 			}
 
-			Encoder = nullptr;
+			Buffer = nullptr;
 		}
 
-		if (Buffer != nullptr)
+		if (Encoder != nullptr)
 		{
-			NVENCSTATUS const Result = FAPI::Get<FNVENC>().nvEncDestroyBitstreamBuffer(Encoder, Buffer);
+			NVENCSTATUS const Result = FAPI::Get<FNVENC>().nvEncDestroyEncoder(Encoder);
 			if (Result != NV_ENC_SUCCESS)
 			{
-				FAVResult::Log(EAVResult::ErrorDestroying, TEXT("Failed to destroy output buffer"), TEXT("NVENC-CUDA"), Result);
+				FAVResult::Log(EAVResult::ErrorDestroying, TEXT("Failed to destroy encoder"), TEXT("NVENC-CUDA"), Result);
 			}
 
-			Buffer = nullptr;
+			Encoder = nullptr;
 		}
 	}
 }
@@ -84,96 +84,93 @@ FAVResult FVideoEncoderNVENCCUDA::SendFrame(TSharedPtr<FVideoResourceCUDA> const
 		return AVResult;
 	}
 
-	NV_ENC_STRUCT(NV_ENC_REGISTER_RESOURCE, InputRegisterResource);
-	NV_ENC_STRUCT(NV_ENC_MAP_INPUT_RESOURCE, InputMapResource);
-	NV_ENC_STRUCT(NV_ENC_PIC_PARAMS, Picture);
-
-	NV_ENC_STRUCT(NV_ENC_LOCK_BITSTREAM, BitstreamLock);
-	BitstreamLock.outputBitstream = Buffer;
-
-	if (Resource.IsValid())
 	{
-		TAVResult<NV_ENC_BUFFER_FORMAT> const ConvertedPixelFormat = FVideoEncoderConfigNVENC::ConvertFormat(Resource->GetFormat());
-		if (ConvertedPixelFormat.IsNotSuccess())
-		{
-			return ConvertedPixelFormat;
-		}
+		NV_ENC_STRUCT(NV_ENC_REGISTER_RESOURCE, InputRegisterResource);
+		NV_ENC_STRUCT(NV_ENC_MAP_INPUT_RESOURCE, InputMapResource);
+		NV_ENC_STRUCT(NV_ENC_PIC_PARAMS, Picture);
 
-		DebugDumpFrame(Resource, Timestamp);
-
-		AVResult = RegisterResource(
-			InputRegisterResource,
-			NV_ENC_INPUT_RESOURCE_TYPE_CUDAARRAY, // ResourceType
-			ConvertedPixelFormat,				  // ResourceFormat
-			Resource->GetRaw(),					  // Resource
-			Resource->GetWidth(),				  // Width
-			Resource->GetHeight(),				  // Height
-			Resource->GetStride(),				  // Stride
-			NV_ENC_INPUT_IMAGE					  // BufferUsage
-		);
-		if (AVResult.IsNotSuccess())
+		ON_SCOPE_EXIT
 		{
-			return AVResult;
-		}
+			if (FAVResult UnmapResult = UnmapResource(InputMapResource.mappedResource); UnmapResult.IsNotSuccess())
+			{
+				UnmapResult.Log();
+			}
+		
+			if (FAVResult UnregisterResult = UnregisterResource(InputRegisterResource.registeredResource); UnregisterResult.IsNotSuccess())
+			{
+				UnregisterResult.Log();
+			}
+		};
 
-		AVResult = MapResource(InputMapResource, InputRegisterResource.registeredResource);
-		if (AVResult.IsNotSuccess())
+		NV_ENC_STRUCT(NV_ENC_LOCK_BITSTREAM, BitstreamLock);
+		BitstreamLock.outputBitstream = Buffer;
+
+		if (Resource.IsValid())
 		{
-			return AVResult;
-		}
+			TAVResult<NV_ENC_BUFFER_FORMAT> const ConvertedPixelFormat = FVideoEncoderConfigNVENC::ConvertFormat(Resource->GetFormat());
+			if (ConvertedPixelFormat.IsNotSuccess())
+			{
+				return ConvertedPixelFormat;
+			}
 
-		Picture.inputWidth = InputRegisterResource.width;
-		Picture.inputHeight = InputRegisterResource.height;
-		Picture.inputPitch = InputRegisterResource.pitch;
-		Picture.inputBuffer = InputMapResource.mappedResource;
-		Picture.bufferFmt = InputMapResource.mappedBufferFmt;
-		Picture.pictureStruct = NV_ENC_PIC_STRUCT_FRAME;
-		Picture.inputTimeStamp = Timestamp;
-		Picture.outputBitstream = Buffer;
+			DebugDumpFrame(Resource, Timestamp);
+
+			AVResult = RegisterResource(
+				InputRegisterResource,
+				NV_ENC_INPUT_RESOURCE_TYPE_CUDAARRAY, // ResourceType
+				ConvertedPixelFormat,				  // ResourceFormat
+				Resource->GetRaw(),					  // Resource
+				Resource->GetWidth(),				  // Width
+				Resource->GetHeight(),				  // Height
+				Resource->GetStride(),				  // Stride
+				NV_ENC_INPUT_IMAGE					  // BufferUsage
+			);
+			if (AVResult.IsNotSuccess())
+			{
+				return AVResult;
+			}
 
-		if (bForceKeyframe)
-		{
-			Picture.encodePicFlags |= NV_ENC_PIC_FLAG_FORCEIDR;
-		}
-	}
-	else
-	{
-		Picture.encodePicFlags = NV_ENC_PIC_FLAG_EOS;
-	}
+			AVResult = MapResource(InputMapResource, InputRegisterResource.registeredResource);
+			if (AVResult.IsNotSuccess())
+			{
+				return AVResult;
+			}
 
-	if (Resource.IsValid())
-	{
-		Resource->Lock();
-	}
+			Picture.inputWidth = InputRegisterResource.width;
+			Picture.inputHeight = InputRegisterResource.height;
+			Picture.inputPitch = InputRegisterResource.pitch;
+			Picture.inputBuffer = InputMapResource.mappedResource;
+			Picture.bufferFmt = InputMapResource.mappedBufferFmt;
+			Picture.pictureStruct = NV_ENC_PIC_STRUCT_FRAME;
+			Picture.inputTimeStamp = Timestamp;
+			Picture.outputBitstream = Buffer;
 
-	AVResult = FEncoderNVENC::SendFrame(Picture, BitstreamLock);
+			if (bForceKeyframe)
+			{
+				Picture.encodePicFlags |= NV_ENC_PIC_FLAG_FORCEIDR;
+			}
+		}
+		else
+		{
+			Picture.encodePicFlags = NV_ENC_PIC_FLAG_EOS;
+		}
 
-	if (Resource.IsValid())
-	{
-		Resource->Unlock();
-	}
+		if (Resource.IsValid())
+		{
+			Resource->Lock();
+		}
 
-	if (AVResult.IsNotSuccess())
-	{
-		// Here we just log the error instead of returning so that we ensure we unmap and unregister the input resources
-		AVResult.Log();
-	}
+		AVResult = FEncoderNVENC::SendFrame(Picture, BitstreamLock);
 
-	if (FAVResult UnmapResult = UnmapResource(InputMapResource.mappedResource); UnmapResult.IsNotSuccess())
-	{
-		UnmapResult.Log();
-		if(UnmapResult < AVResult)
+		if (Resource.IsValid())
 		{
-			AVResult = UnmapResult;
+			Resource->Unlock();
 		}
-	}
 
-	if (FAVResult UnregisterResult = UnregisterResource(InputRegisterResource.registeredResource); UnregisterResult.IsNotSuccess())
-	{
-		UnregisterResult.Log();
-		if (UnregisterResult < AVResult)
+		if (AVResult.IsNotSuccess())
 		{
-			AVResult = UnregisterResult;
+			// Here we just log the error instead of returning so that we ensure we unmap and unregister the input resources
+			AVResult.Log();
 		}
 	}
 
diff --git a/Engine/Plugins/Experimental/AVCodecs/NVCodecs/Source/NVENC/Private/Video/Encoders/VideoEncoderNVENCD3D11.cpp b/Engine/Plugins/Experimental/AVCodecs/NVCodecs/Source/NVENC/Private/Video/Encoders/VideoEncoderNVENCD3D11.cpp
--- a/Engine/Plugins/Experimental/AVCodecs/NVCodecs/Source/NVENC/Private/Video/Encoders/VideoEncoderNVENCD3D11.cpp
+++ b/Engine/Plugins/Experimental/AVCodecs/NVCodecs/Source/NVENC/Private/Video/Encoders/VideoEncoderNVENCD3D11.cpp
@@ -81,26 +81,26 @@ void FVideoEncoderNVENCD3D11::Close()
 {
 	if (IsOpen())
 	{
-		if (Encoder != nullptr)
+		if (Buffer != nullptr)
 		{
-			NVENCSTATUS const Result = FAPI::Get<FNVENC>().nvEncDestroyEncoder(Encoder);
+			NVENCSTATUS const Result = FAPI::Get<FNVENC>().nvEncDestroyBitstreamBuffer(Encoder, Buffer);
 			if (Result != NV_ENC_SUCCESS)
 			{
-				FAVResult::Log(EAVResult::ErrorDestroying, TEXT("Failed to destroy encoder"), TEXT("NVENC-D3D11"), Result);
+				FAVResult::Log(EAVResult::ErrorDestroying, TEXT("Failed to destroy output buffer"), TEXT("NVENC-D3D11"), Result);
 			}
 
-			Encoder = nullptr;
+			Buffer = nullptr;
 		}
 
-		if (Buffer != nullptr)
+		if (Encoder != nullptr)
 		{
-			NVENCSTATUS const Result = FAPI::Get<FNVENC>().nvEncDestroyBitstreamBuffer(Encoder, Buffer);
+			NVENCSTATUS const Result = FAPI::Get<FNVENC>().nvEncDestroyEncoder(Encoder);
 			if (Result != NV_ENC_SUCCESS)
 			{
-				FAVResult::Log(EAVResult::ErrorDestroying, TEXT("Failed to destroy output buffer"), TEXT("NVENC-D3D11"), Result);
+				FAVResult::Log(EAVResult::ErrorDestroying, TEXT("Failed to destroy encoder"), TEXT("NVENC-D3D11"), Result);
 			}
 
-			Buffer = nullptr;
+			Encoder = nullptr;
 		}
 	}
 }
@@ -148,94 +148,91 @@ FAVResult FVideoEncoderNVENCD3D11::SendFrame(TSharedPtr<FVideoResourceD3D11> con
 		return FAVResult(EAVResult::Fatal, TEXT("Failed to open shared handle."), TEXT("NVENC-D3D11"), HResult);
 	}
 
-	NV_ENC_STRUCT(NV_ENC_REGISTER_RESOURCE, InputRegisterResource);
-	NV_ENC_STRUCT(NV_ENC_MAP_INPUT_RESOURCE, InputMapResource);
-	NV_ENC_STRUCT(NV_ENC_PIC_PARAMS, Picture);
-
-	NV_ENC_STRUCT(NV_ENC_LOCK_BITSTREAM, BitstreamLock);
-	BitstreamLock.outputBitstream = Buffer;
-
-	if (Resource.IsValid())
 	{
-		TAVResult<NV_ENC_BUFFER_FORMAT> const ConvertedPixelFormat = FVideoEncoderConfigNVENC::ConvertFormat(Resource->GetFormat());
-		if (ConvertedPixelFormat.IsNotSuccess())
-		{
-			return ConvertedPixelFormat;
-		}
+		NV_ENC_STRUCT(NV_ENC_REGISTER_RESOURCE, InputRegisterResource);
+		NV_ENC_STRUCT(NV_ENC_MAP_INPUT_RESOURCE, InputMapResource);
+		NV_ENC_STRUCT(NV_ENC_PIC_PARAMS, Picture);
 
-		AVResult = RegisterResource(
-			InputRegisterResource,
-			NV_ENC_INPUT_RESOURCE_TYPE_DIRECTX, // ResourceType
-			ConvertedPixelFormat,				// ResourceFormat
-			Resource->GetRaw(),					// Resource
-			Resource->GetWidth(),				// Width
-			Resource->GetHeight(),				// Height
-			Resource->GetStride(),				// Stride
-			NV_ENC_INPUT_IMAGE					// BufferUsage
-		);
-		if (AVResult.IsNotSuccess())
+		ON_SCOPE_EXIT
 		{
-			return AVResult;
-		}
+			if (FAVResult UnmapResult = UnmapResource(InputMapResource.mappedResource); UnmapResult.IsNotSuccess())
+			{
+				UnmapResult.Log();
+			}
+		
+			if (FAVResult UnregisterResult = UnregisterResource(InputRegisterResource.registeredResource); UnregisterResult.IsNotSuccess())
+			{
+				UnregisterResult.Log();
+			}
+		};
 
-		AVResult = MapResource(InputMapResource, InputRegisterResource.registeredResource);
-        if (AVResult.IsNotSuccess())
+		NV_ENC_STRUCT(NV_ENC_LOCK_BITSTREAM, BitstreamLock);
+		BitstreamLock.outputBitstream = Buffer;
+
+		if (Resource.IsValid())
 		{
-			return AVResult;
-		}
+			TAVResult<NV_ENC_BUFFER_FORMAT> const ConvertedPixelFormat = FVideoEncoderConfigNVENC::ConvertFormat(Resource->GetFormat());
+			if (ConvertedPixelFormat.IsNotSuccess())
+			{
+				return ConvertedPixelFormat;
+			}
 
-		Picture.inputWidth = InputRegisterResource.width;
-		Picture.inputHeight = InputRegisterResource.height;
-		Picture.inputPitch = InputRegisterResource.pitch;
-		Picture.inputBuffer = InputMapResource.mappedResource;
-		Picture.bufferFmt = InputMapResource.mappedBufferFmt;
-		Picture.pictureStruct = NV_ENC_PIC_STRUCT_FRAME;
-		Picture.inputTimeStamp = Timestamp;
-		Picture.outputBitstream = Buffer;
+			AVResult = RegisterResource(
+				InputRegisterResource,
+				NV_ENC_INPUT_RESOURCE_TYPE_DIRECTX, // ResourceType
+				ConvertedPixelFormat,				// ResourceFormat
+				Resource->GetRaw(),					// Resource
+				Resource->GetWidth(),				// Width
+				Resource->GetHeight(),				// Height
+				Resource->GetStride(),				// Stride
+				NV_ENC_INPUT_IMAGE					// BufferUsage
+			);
+			if (AVResult.IsNotSuccess())
+			{
+				return AVResult;
+			}
 
-		if (bForceKeyframe)
-		{
-			Picture.encodePicFlags |= NV_ENC_PIC_FLAG_FORCEIDR;
-		}
-	}
-	else
-	{
-		Picture.encodePicFlags = NV_ENC_PIC_FLAG_EOS;
-	}
+			AVResult = MapResource(InputMapResource, InputRegisterResource.registeredResource);
+    	    if (AVResult.IsNotSuccess())
+			{
+				return AVResult;
+			}
 
-	if (Resource.IsValid())
-	{
-		Resource->Lock();
-	}
+			Picture.inputWidth = InputRegisterResource.width;
+			Picture.inputHeight = InputRegisterResource.height;
+			Picture.inputPitch = InputRegisterResource.pitch;
+			Picture.inputBuffer = InputMapResource.mappedResource;
+			Picture.bufferFmt = InputMapResource.mappedBufferFmt;
+			Picture.pictureStruct = NV_ENC_PIC_STRUCT_FRAME;
+			Picture.inputTimeStamp = Timestamp;
+			Picture.outputBitstream = Buffer;
 
-	AVResult = FEncoderNVENC::SendFrame(Picture, BitstreamLock);
+			if (bForceKeyframe)
+			{
+				Picture.encodePicFlags |= NV_ENC_PIC_FLAG_FORCEIDR;
+			}
+		}
+		else
+		{
+			Picture.encodePicFlags = NV_ENC_PIC_FLAG_EOS;
+		}
 
-	if (Resource.IsValid())
-	{
-		Resource->Unlock();
-	}
+		if (Resource.IsValid())
+		{
+			Resource->Lock();
+		}
 
-	if (AVResult.IsNotSuccess())
-	{
-		// Here we just log the error instead of returning so that we ensure we unmap and unregister the input resources
-		AVResult.Log();
-	}
+		AVResult = FEncoderNVENC::SendFrame(Picture, BitstreamLock);
 
-	if (FAVResult UnmapResult = UnmapResource(InputMapResource.mappedResource); UnmapResult.IsNotSuccess())
-	{
-		UnmapResult.Log();
-		if(UnmapResult < AVResult)
+		if (Resource.IsValid())
 		{
-			AVResult = UnmapResult;
+			Resource->Unlock();
 		}
-	}
 
-	if (FAVResult UnregisterResult = UnregisterResource(InputRegisterResource.registeredResource); UnregisterResult.IsNotSuccess())
-	{
-		UnregisterResult.Log();
-		if (UnregisterResult < AVResult)
+		if (AVResult.IsNotSuccess())
 		{
-			AVResult = UnregisterResult;
+			// Here we just log the error instead of returning so that we ensure we unmap and unregister the input resources
+			AVResult.Log();
 		}
 	}
 
diff --git a/Engine/Plugins/Experimental/AVCodecs/NVCodecs/Source/NVENC/Private/Video/Encoders/VideoEncoderNVENCD3D12.cpp b/Engine/Plugins/Experimental/AVCodecs/NVCodecs/Source/NVENC/Private/Video/Encoders/VideoEncoderNVENCD3D12.cpp
--- a/Engine/Plugins/Experimental/AVCodecs/NVCodecs/Source/NVENC/Private/Video/Encoders/VideoEncoderNVENCD3D12.cpp
+++ b/Engine/Plugins/Experimental/AVCodecs/NVCodecs/Source/NVENC/Private/Video/Encoders/VideoEncoderNVENCD3D12.cpp
@@ -118,149 +118,141 @@ FAVResult FVideoEncoderNVENCD3D12::SendFrame(TSharedPtr<FVideoResourceD3D12> con
 
 	FAVResult	AVResult = EAVResult::Success;
 
-	NV_ENC_STRUCT(NV_ENC_REGISTER_RESOURCE, InputRegisterResource);
-	NV_ENC_STRUCT(NV_ENC_MAP_INPUT_RESOURCE, InputMapResource);
-	NV_ENC_STRUCT(NV_ENC_INPUT_RESOURCE_D3D12, InputResource);
-	// Register and map input resource
-	if (Resource.IsValid())
 	{
-		AVResult = RegisterResource(
-			InputRegisterResource,
-			NV_ENC_INPUT_RESOURCE_TYPE_DIRECTX, // ResourceType
-			NV_ENC_BUFFER_FORMAT_ARGB,			// ResourceFormat
-			Resource->GetResource(),			// Resource
-			Resource->GetWidth(),				// Width
-			Resource->GetHeight(),				// Height
-			0,									// Stride
-			NV_ENC_INPUT_IMAGE					// BufferUsage
-		);
-		if (AVResult.IsNotSuccess())
-		{
-			return AVResult;
-		}
+		NV_ENC_STRUCT(NV_ENC_REGISTER_RESOURCE, InputRegisterResource);
+		NV_ENC_STRUCT(NV_ENC_MAP_INPUT_RESOURCE, InputMapResource);
+		NV_ENC_STRUCT(NV_ENC_REGISTER_RESOURCE, OutputRegisterResource);
+		NV_ENC_STRUCT(NV_ENC_MAP_INPUT_RESOURCE, OutputMapResource);
 
-		AVResult = MapResource(InputMapResource, InputRegisterResource.registeredResource);
-		if (AVResult.IsNotSuccess())
+		ON_SCOPE_EXIT
 		{
-			return AVResult;
-		}
-
-		InputResource.inputFencePoint.pFence = InputFence.GetReference();
-		InputResource.pInputBuffer = InputMapResource.mappedResource;
-		InputResource.inputFencePoint.waitValue = InputFenceVal;
-		InputResource.inputFencePoint.bWait = true;
-	}
+			// Unmap and unregister input resources
+			if (FAVResult UnmapResult = UnmapResource(InputMapResource.mappedResource); UnmapResult.IsNotSuccess())
+			{
+				UnmapResult.Log();
+			}
 
-	FPlatformAtomics::InterlockedIncrement(&OutputFenceVal);
-	// Register and map output resource
-	NV_ENC_STRUCT(NV_ENC_REGISTER_RESOURCE, OutputRegisterResource);
-	NV_ENC_STRUCT(NV_ENC_MAP_INPUT_RESOURCE, OutputMapResource);
-	NV_ENC_STRUCT(NV_ENC_OUTPUT_RESOURCE_D3D12, OutputResource);
-	{
-		AVResult = RegisterResource(
-			OutputRegisterResource,
-			NV_ENC_INPUT_RESOURCE_TYPE_DIRECTX,										  // ResourceType
-			NV_ENC_BUFFER_FORMAT_U8,												  // ResourceFormat
-			OutputBitstreamResource.GetReference(),									  // Resource
-			Align(AppliedConfig.encodeWidth * AppliedConfig.encodeHeight * 4 * 2, 4), // Width
-			1,																		  // Height
-			0,																		  // Stride
-			NV_ENC_OUTPUT_BITSTREAM													  // BufferUsage
-		);
-		if (AVResult.IsNotSuccess())
-		{
-			return AVResult;
-		}
+			if (FAVResult UnregisterResult = UnregisterResource(InputRegisterResource.registeredResource); UnregisterResult.IsNotSuccess())
+			{
+				UnregisterResult.Log();
+			}
 
-		AVResult = MapResource(OutputMapResource, OutputRegisterResource.registeredResource);
-		if (AVResult.IsNotSuccess())
-		{
-			return AVResult;
-		}
+			// Unmap and unregister output resources
+			if (FAVResult UnmapResult = UnmapResource(OutputMapResource.mappedResource); UnmapResult.IsNotSuccess())
+			{
+				UnmapResult.Log();
+			}
 
-		OutputResource.outputFencePoint.pFence = OutputFence.GetReference();
-		OutputResource.pOutputBuffer = OutputMapResource.mappedResource;
-		OutputResource.outputFencePoint.signalValue = OutputFenceVal;
-		OutputResource.outputFencePoint.bSignal = true;
-	}
+			if (FAVResult UnregisterResult = UnregisterResource(OutputRegisterResource.registeredResource); UnregisterResult.IsNotSuccess())
+			{
+				UnregisterResult.Log();
+			}
+		};
 
-	NV_ENC_STRUCT(NV_ENC_PIC_PARAMS, Picture);
-	{
-		Picture.pictureStruct = NV_ENC_PIC_STRUCT_FRAME;
-		Picture.inputTimeStamp = Timestamp;
-		Picture.inputBuffer = &InputResource;
-		Picture.bufferFmt = InputMapResource.mappedBufferFmt;
-		Picture.inputWidth = InputRegisterResource.width;
-		Picture.inputHeight = InputRegisterResource.height;
-		Picture.outputBitstream = &OutputResource;
-		Picture.completionEvent = nullptr;
-	}
+		// Register and map input resource
+		NV_ENC_STRUCT(NV_ENC_INPUT_RESOURCE_D3D12, InputResource);
+		if (Resource.IsValid())
+		{
+			AVResult = RegisterResource(
+				InputRegisterResource,
+				NV_ENC_INPUT_RESOURCE_TYPE_DIRECTX, // ResourceType
+				NV_ENC_BUFFER_FORMAT_ARGB,			// ResourceFormat
+				Resource->GetResource(),			// Resource
+				Resource->GetWidth(),				// Width
+				Resource->GetHeight(),				// Height
+				0,									// Stride
+				NV_ENC_INPUT_IMAGE					// BufferUsage
+			);
+			if (AVResult.IsNotSuccess())
+			{
+				return AVResult;
+			}
 
-	if (!Resource.IsValid())
-	{
-		Picture.encodePicFlags = NV_ENC_PIC_FLAG_EOS;
-	}
-	else if (bForceKeyframe)
-	{
-		Picture.encodePicFlags |= NV_ENC_PIC_FLAG_FORCEIDR;
-	}
+			AVResult = MapResource(InputMapResource, InputRegisterResource.registeredResource);
+			if (AVResult.IsNotSuccess())
+			{
+				return AVResult;
+			}
 
-	if (Resource.IsValid())
-	{
-		Resource->Lock();
-	}
+			InputResource.inputFencePoint.pFence = InputFence.GetReference();
+			InputResource.pInputBuffer = InputMapResource.mappedResource;
+			InputResource.inputFencePoint.waitValue = InputFenceVal;
+			InputResource.inputFencePoint.bWait = true;
+		}
 
-	NV_ENC_STRUCT(NV_ENC_LOCK_BITSTREAM, BitstreamLock);
-	BitstreamLock.outputBitstream = &OutputResource;
+		FPlatformAtomics::InterlockedIncrement(&OutputFenceVal);
 
-	AVResult = FEncoderNVENC::SendFrame(Picture, BitstreamLock);
+		// Register and map output resource
+		NV_ENC_STRUCT(NV_ENC_OUTPUT_RESOURCE_D3D12, OutputResource);
+		{
+			AVResult = RegisterResource(
+				OutputRegisterResource,
+				NV_ENC_INPUT_RESOURCE_TYPE_DIRECTX,										  // ResourceType
+				NV_ENC_BUFFER_FORMAT_U8,												  // ResourceFormat
+				OutputBitstreamResource.GetReference(),									  // Resource
+				Align(AppliedConfig.encodeWidth * AppliedConfig.encodeHeight * 4 * 2, 4), // Width
+				1,																		  // Height
+				0,																		  // Stride
+				NV_ENC_OUTPUT_BITSTREAM													  // BufferUsage
+			);
+			if (AVResult.IsNotSuccess())
+			{
+				return AVResult;
+			}
 
-	if (Resource.IsValid())
-	{
-		Resource->Unlock();
-	}
+			AVResult = MapResource(OutputMapResource, OutputRegisterResource.registeredResource);
+			if (AVResult.IsNotSuccess())
+			{
+				return AVResult;
+			}
 
-	if (AVResult.IsNotSuccess())
-	{
-		// Here we just log the error instead of returning so that we ensure we unmap and unregister the input resources
-		AVResult.Log();
-	}
+			OutputResource.outputFencePoint.pFence = OutputFence.GetReference();
+			OutputResource.pOutputBuffer = OutputMapResource.mappedResource;
+			OutputResource.outputFencePoint.signalValue = OutputFenceVal;
+			OutputResource.outputFencePoint.bSignal = true;
+		}
 
-	// Unmap and unregister input resources
-	if (FAVResult UnmapResult = UnmapResource(InputMapResource.mappedResource); UnmapResult.IsNotSuccess())
-	{
-		UnmapResult.Log();
-		if(UnmapResult < AVResult)
+		NV_ENC_STRUCT(NV_ENC_PIC_PARAMS, Picture);
+		if (Resource.IsValid())
+		{
+			Picture.pictureStruct = NV_ENC_PIC_STRUCT_FRAME;
+			Picture.inputTimeStamp = Timestamp;
+			Picture.inputBuffer = &InputResource;
+			Picture.bufferFmt = InputMapResource.mappedBufferFmt;
+			Picture.inputWidth = InputRegisterResource.width;
+			Picture.inputHeight = InputRegisterResource.height;
+			Picture.outputBitstream = &OutputResource;
+			Picture.completionEvent = nullptr;
+
+			if (bForceKeyframe)
+			{
+				Picture.encodePicFlags |= NV_ENC_PIC_FLAG_FORCEIDR;
+			}
+		}
+		else
 		{
-			AVResult = UnmapResult;
+			Picture.encodePicFlags = NV_ENC_PIC_FLAG_EOS;
 		}
-	}
 
-	if (FAVResult UnregisterResult = UnregisterResource(InputRegisterResource.registeredResource); UnregisterResult.IsNotSuccess())
-	{
-		UnregisterResult.Log();
-		if (UnregisterResult < AVResult)
+		if (Resource.IsValid())
 		{
-			AVResult = UnregisterResult;
+			Resource->Lock();
 		}
-	}
 
-	// Unmap and unregister output resources
-	if (FAVResult UnmapResult = UnmapResource(OutputMapResource.mappedResource); UnmapResult.IsNotSuccess())
-	{
-		UnmapResult.Log();
-		if(UnmapResult < AVResult)
+		NV_ENC_STRUCT(NV_ENC_LOCK_BITSTREAM, BitstreamLock);
+		BitstreamLock.outputBitstream = &OutputResource;
+
+		AVResult = FEncoderNVENC::SendFrame(Picture, BitstreamLock);
+
+		if (Resource.IsValid())
 		{
-			AVResult = UnmapResult;
+			Resource->Unlock();
 		}
-	}
 
-	if (FAVResult UnregisterResult = UnregisterResource(OutputRegisterResource.registeredResource); UnregisterResult.IsNotSuccess())
-	{
-		UnregisterResult.Log();
-		if (UnregisterResult < AVResult)
+		if (AVResult.IsNotSuccess())
 		{
-			AVResult = UnregisterResult;
+			// Here we just log the error instead of returning so that we ensure we unmap and unregister the input resources
+			AVResult.Log();
 		}
 	}
 
diff --git a/Engine/Plugins/Experimental/AVCodecs/NVCodecs/Source/NVENC/Public/Video/Encoders/VideoEncoderNVENC.h b/Engine/Plugins/Experimental/AVCodecs/NVCodecs/Source/NVENC/Public/Video/Encoders/VideoEncoderNVENC.h
--- a/Engine/Plugins/Experimental/AVCodecs/NVCodecs/Source/NVENC/Public/Video/Encoders/VideoEncoderNVENC.h
+++ b/Engine/Plugins/Experimental/AVCodecs/NVCodecs/Source/NVENC/Public/Video/Encoders/VideoEncoderNVENC.h
@@ -2,14 +2,15 @@
 
 #pragma once
 
-#include "Video/VideoEncoder.h"
-#include "Video/Resources/VideoResourceCUDA.h"
-#include "Video/VideoResource.h"
+#include "Containers/Queue.h"
+#include "Misc/ScopeExit.h"
 #include "NVENC.h"
-#include "Video/Encoders/Configs/VideoEncoderConfigNVENC.h"
 #include "Templates/AlignmentTemplates.h"
 #include "Templates/RefCounting.h"
-#include "Containers/Queue.h"
+#include "Video/Encoders/Configs/VideoEncoderConfigNVENC.h"
+#include "Video/Resources/VideoResourceCUDA.h"
+#include "Video/VideoEncoder.h"
+#include "Video/VideoResource.h"
 
 #if PLATFORM_WINDOWS
 	#include "Video/Resources/D3D/VideoResourceD3D.h"
@@ -23,8 +24,7 @@ namespace AVUtils
 #endif
 
 /**
- * This class is the core of the NVENC encoder logic.
- * It is embedded inside NVCodecEncoders using composition instead of inheritance because templates + inheritance was bit messy here where we need concrete types.
+ * This class is the core of the NVENC encoder logic. Specialisations should inherit from this class to reduce duplicated logic
  **/
 class NVENC_API FEncoderNVENC
 {
