The source files included were modified by the UnrealCodeFixup tool so that they can pass the -mergemodules compilation.

diff --git a/Engine/Plugins/Experimental/AVCodecs/NVCodecs/Source/NVCodecs/Public/Video/Resources/VideoResourceCUDA.h b/Engine/Plugins/Experimental/AVCodecs/NVCodecs/Source/NVCodecs/Public/Video/Resources/VideoResourceCUDA.h
--- a/Engine/Plugins/Experimental/AVCodecs/NVCodecs/Source/NVCodecs/Public/Video/Resources/VideoResourceCUDA.h
+++ b/Engine/Plugins/Experimental/AVCodecs/NVCodecs/Source/NVCodecs/Public/Video/Resources/VideoResourceCUDA.h
@@ -8,6 +8,8 @@
 
 #include "CudaModule.h"
 
+#define UE_API NVCODECS_API
+
 struct FCUDAContextScope
 {
 	FCUDAContextScope(CUcontext Context)
@@ -21,15 +23,15 @@ struct FCUDAContextScope
 	}
 };
 
-class NVCODECS_API FVideoContextCUDA : public FAVContext
+class FVideoContextCUDA : public FAVContext
 {
 public:
 	CUcontext Raw;
 
-	FVideoContextCUDA(CUcontext const& Raw);
+	UE_API FVideoContextCUDA(CUcontext const& Raw);
 };
 
-class NVCODECS_API FVideoResourceCUDA : public TVideoResource<FVideoContextCUDA>
+class FVideoResourceCUDA : public TVideoResource<FVideoContextCUDA>
 {
 private:
 	CUexternalSemaphore ExternalSemaphore = nullptr;
@@ -43,20 +45,20 @@ private:
 public:
 	inline CUarray GetRaw() const { return Raw; }
 
-	FVideoResourceCUDA(TSharedRef<FAVDevice> const& Device, CUarray Raw, FAVLayout const& Layout, FVideoDescriptor const& Descriptor);
-	FVideoResourceCUDA(TSharedRef<FAVDevice> const& Device, CUDA_EXTERNAL_MEMORY_HANDLE_DESC const& ExternalResourceDesc, CUDA_EXTERNAL_SEMAPHORE_HANDLE_DESC const& ExternalFenceDesc, FAVLayout const& Layout, FVideoDescriptor const& Descriptor, uint64 FenceValue = 0);
-	FVideoResourceCUDA(TSharedRef<FAVDevice> const& Device, CUDA_EXTERNAL_MEMORY_HANDLE_DESC const& ExternalResourceDesc, CUDA_EXTERNAL_SEMAPHORE_HANDLE_DESC const& ExternalFenceDesc, TSharedRef<FVideoResource> const& ExternalResource, uint64 FenceValue = 0);
-	virtual ~FVideoResourceCUDA() override;
+	UE_API FVideoResourceCUDA(TSharedRef<FAVDevice> const& Device, CUarray Raw, FAVLayout const& Layout, FVideoDescriptor const& Descriptor);
+	UE_API FVideoResourceCUDA(TSharedRef<FAVDevice> const& Device, CUDA_EXTERNAL_MEMORY_HANDLE_DESC const& ExternalResourceDesc, CUDA_EXTERNAL_SEMAPHORE_HANDLE_DESC const& ExternalFenceDesc, FAVLayout const& Layout, FVideoDescriptor const& Descriptor, uint64 FenceValue = 0);
+	UE_API FVideoResourceCUDA(TSharedRef<FAVDevice> const& Device, CUDA_EXTERNAL_MEMORY_HANDLE_DESC const& ExternalResourceDesc, CUDA_EXTERNAL_SEMAPHORE_HANDLE_DESC const& ExternalFenceDesc, TSharedRef<FVideoResource> const& ExternalResource, uint64 FenceValue = 0);
+	UE_API virtual ~FVideoResourceCUDA() override;
 
-	virtual FAVResult Validate() const override;
+	UE_API virtual FAVResult Validate() const override;
 
-	FAVResult ReadData(TArray64<uint8>& OutData) const;
+	UE_API FAVResult ReadData(TArray64<uint8>& OutData) const;
 
-	FAVResult CopyFrom(CUdeviceptr Target, uint32 MapPitch);
-	FAVResult CopyFromAsync(CUdeviceptr Target, uint32 MapPitch);
+	UE_API FAVResult CopyFrom(CUdeviceptr Target, uint32 MapPitch);
+	UE_API FAVResult CopyFromAsync(CUdeviceptr Target, uint32 MapPitch);
 
-	FAVResult CopyTo(CUdeviceptr Source, uint32 MapPitch);
-	FAVResult CopyToAsync(CUdeviceptr Source, uint32 MapPitch);
+	UE_API FAVResult CopyTo(CUdeviceptr Source, uint32 MapPitch);
+	UE_API FAVResult CopyToAsync(CUdeviceptr Source, uint32 MapPitch);
 
 private:
 	inline CUDA_MEMCPY2D CreateCopyFromParams(CUdeviceptr Source, uint32 SourcePitch)
@@ -102,4 +104,6 @@ template <>
 FAVResult FAVExtension::TransformResource(TSharedPtr<FVideoResourceCUDA>& OutResource, TSharedPtr<class FVideoResourceVulkan> const& InResource);
 
 DECLARE_TYPEID(FVideoContextCUDA, NVCODECS_API);
-DECLARE_TYPEID(FVideoResourceCUDA, NVCODECS_API);
\ No newline at end of file
+DECLARE_TYPEID(FVideoResourceCUDA, NVCODECS_API);
+
+#undef UE_API
diff --git a/Engine/Plugins/Experimental/AVCodecs/NVCodecs/Source/NVDEC/Public/NVDEC.h b/Engine/Plugins/Experimental/AVCodecs/NVCodecs/Source/NVDEC/Public/NVDEC.h
--- a/Engine/Plugins/Experimental/AVCodecs/NVCodecs/Source/NVDEC/Public/NVDEC.h
+++ b/Engine/Plugins/Experimental/AVCodecs/NVCodecs/Source/NVDEC/Public/NVDEC.h
@@ -55,19 +55,19 @@ struct NV_DECODE_API_FUNCTION_LIST
 	cuvidUnmapVideoFramePtr cuvidUnmapVideoFrame = nullptr;
 };
 
-class NVDEC_API FNVDEC : public FAPI, public NV_DECODE_API_FUNCTION_LIST
+class FNVDEC : public FAPI, public NV_DECODE_API_FUNCTION_LIST
 {
 public:
 	bool bHasCompatibleGPU = true;
 
-	FNVDEC();
+	NVDEC_API FNVDEC();
 
-	virtual bool IsValid() const override;
+	NVDEC_API virtual bool IsValid() const override;
 
-	FString GetErrorString(CUresult ErrorCode) const;
+	NVDEC_API FString GetErrorString(CUresult ErrorCode) const;
 
 private:
 	void* DllHandle = nullptr;
 };
 
-DECLARE_TYPEID(FNVDEC, NVDEC_API);
\ No newline at end of file
+DECLARE_TYPEID(FNVDEC, NVDEC_API);
diff --git a/Engine/Plugins/Experimental/AVCodecs/NVCodecs/Source/NVDEC/Public/Video/Decoders/Configs/VideoDecoderConfigNVDEC.h b/Engine/Plugins/Experimental/AVCodecs/NVCodecs/Source/NVDEC/Public/Video/Decoders/Configs/VideoDecoderConfigNVDEC.h
--- a/Engine/Plugins/Experimental/AVCodecs/NVCodecs/Source/NVDEC/Public/Video/Decoders/Configs/VideoDecoderConfigNVDEC.h
+++ b/Engine/Plugins/Experimental/AVCodecs/NVCodecs/Source/NVDEC/Public/Video/Decoders/Configs/VideoDecoderConfigNVDEC.h
@@ -10,7 +10,7 @@
 
 #include "NVDEC.h"
 
-struct NVDEC_API FVideoDecoderConfigNVDEC : public CUVIDDECODECREATEINFO, public FAVConfig
+struct FVideoDecoderConfigNVDEC : public CUVIDDECODECREATEINFO, public FAVConfig
 {
 public:
 	struct FParsedPicture : public CUVIDPICPARAMS
diff --git a/Engine/Plugins/Experimental/AVCodecs/NVCodecs/Source/NVDEC/Public/Video/Decoders/VideoDecoderNVDEC.h b/Engine/Plugins/Experimental/AVCodecs/NVCodecs/Source/NVDEC/Public/Video/Decoders/VideoDecoderNVDEC.h
--- a/Engine/Plugins/Experimental/AVCodecs/NVCodecs/Source/NVDEC/Public/Video/Decoders/VideoDecoderNVDEC.h
+++ b/Engine/Plugins/Experimental/AVCodecs/NVCodecs/Source/NVDEC/Public/Video/Decoders/VideoDecoderNVDEC.h
@@ -11,7 +11,9 @@
 #include "NVDEC.h"
 #include "Video/Decoders/Configs/VideoDecoderConfigNVDEC.h"
 
-class NVDEC_API FVideoDecoderNVDEC : public TVideoDecoder<FVideoResourceCUDA, FVideoDecoderConfigNVDEC>
+#define UE_API NVDEC_API
+
+class FVideoDecoderNVDEC : public TVideoDecoder<FVideoResourceCUDA, FVideoDecoderConfigNVDEC>
 {
 private:
 	uint8 bIsOpen : 1;
@@ -34,26 +36,26 @@ private:
 	
 public:
 	FVideoDecoderNVDEC() = default;
-	virtual ~FVideoDecoderNVDEC() override;
+	UE_API virtual ~FVideoDecoderNVDEC() override;
 
-	virtual bool IsOpen() const override;
-	virtual FAVResult Open(TSharedRef<FAVDevice> const& NewDevice, TSharedRef<FAVInstance> const& NewInstance) override;
-	virtual void Close() override;
+	UE_API virtual bool IsOpen() const override;
+	UE_API virtual FAVResult Open(TSharedRef<FAVDevice> const& NewDevice, TSharedRef<FAVInstance> const& NewInstance) override;
+	UE_API virtual void Close() override;
 
-	bool IsInitialized() const;
+	UE_API bool IsInitialized() const;
 
-	virtual FAVResult ApplyConfig() override;
+	UE_API virtual FAVResult ApplyConfig() override;
 
-	bool GetCapability(CUVIDDECODECAPS& CapsToQuery) const;
+	UE_API bool GetCapability(CUVIDDECODECAPS& CapsToQuery) const;
 
-	virtual FAVResult SendPacket(FVideoPacket const& Packet) override;
+	UE_API virtual FAVResult SendPacket(FVideoPacket const& Packet) override;
 	
-	virtual FAVResult ReceiveFrame(TResolvableVideoResource<FVideoResourceCUDA>& InOutResource) override;
+	UE_API virtual FAVResult ReceiveFrame(TResolvableVideoResource<FVideoResourceCUDA>& InOutResource) override;
 
 public:
-	int HandleVideoSequence(CUVIDEOFORMAT *VideoFormat);
-    int HandlePictureDecode(CUVIDPICPARAMS *PicParams);
-    int HandlePictureDisplay(CUVIDPARSERDISPINFO *DispInfo);
+	UE_API int HandleVideoSequence(CUVIDEOFORMAT *VideoFormat);
+    UE_API int HandlePictureDecode(CUVIDPICPARAMS *PicParams);
+    UE_API int HandlePictureDisplay(CUVIDPARSERDISPINFO *DispInfo);
 };
 
 namespace Internal 
@@ -62,3 +64,5 @@ namespace Internal
 	int HandlePictureDecodeCallback(void *UserData, CUVIDPICPARAMS *PicParams);
 	int HandlePictureDisplayCallback(void *UserData, CUVIDPARSERDISPINFO *DispInfo);
 }
+
+#undef UE_API
diff --git a/Engine/Plugins/Experimental/AVCodecs/NVCodecs/Source/NVENC/Public/NVENC.h b/Engine/Plugins/Experimental/AVCodecs/NVCodecs/Source/NVENC/Public/NVENC.h
--- a/Engine/Plugins/Experimental/AVCodecs/NVCodecs/Source/NVENC/Public/NVENC.h
+++ b/Engine/Plugins/Experimental/AVCodecs/NVCodecs/Source/NVENC/Public/NVENC.h
@@ -23,16 +23,16 @@ THIRD_PARTY_INCLUDES_END
 	FMemory::Memzero(VarName);			\
 	VarName.version = OfType ## _VER
 
-class NVENC_API FNVENC : public FAPI, public NV_ENCODE_API_FUNCTION_LIST
+class FNVENC : public FAPI, public NV_ENCODE_API_FUNCTION_LIST
 {
 public:
 	bool bHasCompatibleGPU = true;
 
-	FNVENC();
+	NVENC_API FNVENC();
 
-	virtual bool IsValid() const override;
+	NVENC_API virtual bool IsValid() const override;
 
-	FString GetErrorString(void* Encoder, NVENCSTATUS ForStatus) const;
+	NVENC_API FString GetErrorString(void* Encoder, NVENCSTATUS ForStatus) const;
 
 private:
 	void* DllHandle = nullptr;
diff --git a/Engine/Plugins/Experimental/AVCodecs/NVCodecs/Source/NVENC/Public/Video/Encoders/Configs/VideoEncoderConfigNVENC.h b/Engine/Plugins/Experimental/AVCodecs/NVCodecs/Source/NVENC/Public/Video/Encoders/Configs/VideoEncoderConfigNVENC.h
--- a/Engine/Plugins/Experimental/AVCodecs/NVCodecs/Source/NVENC/Public/Video/Encoders/Configs/VideoEncoderConfigNVENC.h
+++ b/Engine/Plugins/Experimental/AVCodecs/NVCodecs/Source/NVENC/Public/Video/Encoders/Configs/VideoEncoderConfigNVENC.h
@@ -7,14 +7,16 @@
 
 #include "NVENC.h"
 
-struct NVENC_API FVideoEncoderConfigNVENC : public FAVConfig, public NV_ENC_INITIALIZE_PARAMS
+#define UE_API NVENC_API
+
+struct FVideoEncoderConfigNVENC : public FAVConfig, public NV_ENC_INITIALIZE_PARAMS
 {
 public:
-	static TAVResult<NV_ENC_PARAMS_RC_MODE> ConvertRateControlMode(ERateControlMode Mode);
-	static TAVResult<ERateControlMode> ConvertRateControlMode(NV_ENC_PARAMS_RC_MODE Mode);
-	static TAVResult<NV_ENC_MULTI_PASS> ConvertMultipassMode(EMultipassMode Mode);
-	static TAVResult<EMultipassMode> ConvertMultipassMode(NV_ENC_MULTI_PASS Mode);
-	static TAVResult<NV_ENC_BUFFER_FORMAT> ConvertFormat(EVideoFormat const& Format);
+	static UE_API TAVResult<NV_ENC_PARAMS_RC_MODE> ConvertRateControlMode(ERateControlMode Mode);
+	static UE_API TAVResult<ERateControlMode> ConvertRateControlMode(NV_ENC_PARAMS_RC_MODE Mode);
+	static UE_API TAVResult<NV_ENC_MULTI_PASS> ConvertMultipassMode(EMultipassMode Mode);
+	static UE_API TAVResult<EMultipassMode> ConvertMultipassMode(NV_ENC_MULTI_PASS Mode);
+	static UE_API TAVResult<NV_ENC_BUFFER_FORMAT> ConvertFormat(EVideoFormat const& Format);
 	
 private:
 	NV_ENC_CONFIG Config = {};
@@ -105,3 +107,5 @@ template <>
 FAVResult FAVExtension::TransformConfig(FVideoEncoderConfigNVENC& OutConfig, struct FVideoEncoderConfigAV1 const& InConfig);
 
 DECLARE_TYPEID(FVideoEncoderConfigNVENC, NVENC_API);
+
+#undef UE_API
diff --git a/Engine/Plugins/Experimental/AVCodecs/NVCodecs/Source/NVENC/Public/Video/Encoders/VideoEncoderNVENC.h b/Engine/Plugins/Experimental/AVCodecs/NVCodecs/Source/NVENC/Public/Video/Encoders/VideoEncoderNVENC.h
--- a/Engine/Plugins/Experimental/AVCodecs/NVCodecs/Source/NVENC/Public/Video/Encoders/VideoEncoderNVENC.h
+++ b/Engine/Plugins/Experimental/AVCodecs/NVCodecs/Source/NVENC/Public/Video/Encoders/VideoEncoderNVENC.h
@@ -26,15 +26,15 @@ namespace AVUtils
 /**
  * This class is the core of the NVENC encoder logic. Specialisations should inherit from this class to reduce duplicated logic
  **/
-class NVENC_API FEncoderNVENC
+class FEncoderNVENC
 {
 private:
 	uint32 MaxDeviceEncodeWidth = 0;
 	uint32 MaxDeviceEncodeHeight = 0;
 	bool bHasMaxDeviceResolution = false;
 
-	void SetMaxResolution(FVideoEncoderConfigNVENC& PendingConfig);
-	void GetMaxDeviceEncodeResolution(const FVideoEncoderConfigNVENC& PendingConfig);
+	NVENC_API void SetMaxResolution(FVideoEncoderConfigNVENC& PendingConfig);
+	NVENC_API void GetMaxDeviceEncodeResolution(const FVideoEncoderConfigNVENC& PendingConfig);
 
 public:
 	FEncoderNVENC() = default;
@@ -47,11 +47,11 @@ protected:
 	void* Encoder = nullptr;
 	TQueue<FVideoPacket> Packets;
 
-	FAVResult ReOpen();
-	FAVResult Open(TSharedRef<FAVDevice> const& NewDevice, TSharedRef<FAVInstance> const& NewInstance, NV_ENC_OPEN_ENCODE_SESSION_EX_PARAMS& EncoderSessionParams);
-	FAVResult ApplyConfig(FVideoEncoderConfigNVENC const& AppliedConfig, FVideoEncoderConfigNVENC& PendingConfig);
-	FAVResult SendFrame(NV_ENC_PIC_PARAMS& Input, NV_ENC_LOCK_BITSTREAM& BitstreamLock);
-	FAVResult ReceivePacket(FVideoPacket& OutPacket);
+	NVENC_API FAVResult ReOpen();
+	NVENC_API FAVResult Open(TSharedRef<FAVDevice> const& NewDevice, TSharedRef<FAVInstance> const& NewInstance, NV_ENC_OPEN_ENCODE_SESSION_EX_PARAMS& EncoderSessionParams);
+	NVENC_API FAVResult ApplyConfig(FVideoEncoderConfigNVENC const& AppliedConfig, FVideoEncoderConfigNVENC& PendingConfig);
+	NVENC_API FAVResult SendFrame(NV_ENC_PIC_PARAMS& Input, NV_ENC_LOCK_BITSTREAM& BitstreamLock);
+	NVENC_API FAVResult ReceivePacket(FVideoPacket& OutPacket);
 
 	//
 	virtual bool IsOpen() const = 0;
@@ -60,49 +60,49 @@ protected:
 	virtual FAVResult OnPostInitialization(FVideoEncoderConfigNVENC& PendingConfig) { return EAVResult::Success; }
 	//
 
-	FAVResult RegisterResource(NV_ENC_REGISTER_RESOURCE& OutRegisterResource, NV_ENC_INPUT_RESOURCE_TYPE ResourceType, NV_ENC_BUFFER_FORMAT ResourceFormat, void* Resource, uint32 Width, uint32 Height, uint32 Stride, NV_ENC_BUFFER_USAGE BufferUsage);
-	FAVResult UnregisterResource(NV_ENC_REGISTERED_PTR RegisteredResource);
-	FAVResult MapResource(NV_ENC_MAP_INPUT_RESOURCE& OutMapResource, NV_ENC_REGISTERED_PTR RegisteredResource);
-	FAVResult UnmapResource(NV_ENC_INPUT_PTR MappedResource);
+	NVENC_API FAVResult RegisterResource(NV_ENC_REGISTER_RESOURCE& OutRegisterResource, NV_ENC_INPUT_RESOURCE_TYPE ResourceType, NV_ENC_BUFFER_FORMAT ResourceFormat, void* Resource, uint32 Width, uint32 Height, uint32 Stride, NV_ENC_BUFFER_USAGE BufferUsage);
+	NVENC_API FAVResult UnregisterResource(NV_ENC_REGISTERED_PTR RegisteredResource);
+	NVENC_API FAVResult MapResource(NV_ENC_MAP_INPUT_RESOURCE& OutMapResource, NV_ENC_REGISTERED_PTR RegisteredResource);
+	NVENC_API FAVResult UnmapResource(NV_ENC_INPUT_PTR MappedResource);
 
-	int GetCapability(GUID EncodeGUID, NV_ENC_CAPS CapsToQuery) const;
+	NVENC_API int GetCapability(GUID EncodeGUID, NV_ENC_CAPS CapsToQuery) const;
 };
 
 /**
  *  An AVCodecs VideoEncoder that takes a FVideoResourceCUDA.
  **/
-class NVENC_API FVideoEncoderNVENCCUDA : public TVideoEncoder<FVideoResourceCUDA, FVideoEncoderConfigNVENC>, public FEncoderNVENC
+class FVideoEncoderNVENCCUDA : public TVideoEncoder<FVideoResourceCUDA, FVideoEncoderConfigNVENC>, public FEncoderNVENC
 {
 private:
 	NV_ENC_OUTPUT_PTR Buffer = nullptr;
 
 public:
 	FVideoEncoderNVENCCUDA() = default;
-	virtual ~FVideoEncoderNVENCCUDA() override;
+	NVENC_API virtual ~FVideoEncoderNVENCCUDA() override;
 
 	// Begin TVideoEncoder interface
-	virtual FAVResult Open(TSharedRef<FAVDevice> const& NewDevice, TSharedRef<FAVInstance> const& NewInstance) override;
-	virtual void Close() override;
-	virtual bool IsOpen() const override;
-	virtual FAVResult ApplyConfig() override;
-	virtual FAVResult SendFrame(TSharedPtr<FVideoResourceCUDA> const& Resource, uint32 Timestamp, bool bForceKeyframe = false) override;
-	virtual FAVResult ReceivePacket(FVideoPacket& OutPacket) override;
+	NVENC_API virtual FAVResult Open(TSharedRef<FAVDevice> const& NewDevice, TSharedRef<FAVInstance> const& NewInstance) override;
+	NVENC_API virtual void Close() override;
+	NVENC_API virtual bool IsOpen() const override;
+	NVENC_API virtual FAVResult ApplyConfig() override;
+	NVENC_API virtual FAVResult SendFrame(TSharedPtr<FVideoResourceCUDA> const& Resource, uint32 Timestamp, bool bForceKeyframe = false) override;
+	NVENC_API virtual FAVResult ReceivePacket(FVideoPacket& OutPacket) override;
 	// End TVideoEncoder interface
 
 	// Begin FEncoderNVENC interface
-	virtual	bool IsInitialized() const override;
-	virtual FAVResult OnPostInitialization(FVideoEncoderConfigNVENC& PendingConfig) override;
+	NVENC_API virtual	bool IsInitialized() const override;
+	NVENC_API virtual FAVResult OnPostInitialization(FVideoEncoderConfigNVENC& PendingConfig) override;
 	// End FEncoderNVENC interface
 
 private:
-	void DebugDumpFrame(TSharedPtr<FVideoResourceCUDA> const& Resource, uint32 Timestamp);
+	NVENC_API void DebugDumpFrame(TSharedPtr<FVideoResourceCUDA> const& Resource, uint32 Timestamp);
 };
 
 #if PLATFORM_WINDOWS
 /**
  *  An AVCodecs VideoEncoder that takes a FVideoResourceD3D11.
  **/
-class NVENC_API FVideoEncoderNVENCD3D11 : public TVideoEncoder<FVideoResourceD3D11, FVideoEncoderConfigNVENC>, public FEncoderNVENC
+class FVideoEncoderNVENCD3D11 : public TVideoEncoder<FVideoResourceD3D11, FVideoEncoderConfigNVENC>, public FEncoderNVENC
 {
 private:
 	TRefCountPtr<ID3D11Device> EncoderDevice;
@@ -112,27 +112,27 @@ private:
 
 public:
 	FVideoEncoderNVENCD3D11() = default;
-	virtual ~FVideoEncoderNVENCD3D11() override;
+	NVENC_API virtual ~FVideoEncoderNVENCD3D11() override;
 
 	// Begin TVideoEncoder interface
-	virtual FAVResult Open(TSharedRef<FAVDevice> const& NewDevice, TSharedRef<FAVInstance> const& NewInstance) override;
-	virtual void Close() override;
-	virtual bool IsOpen() const override;
-	virtual FAVResult ApplyConfig() override;
-	virtual FAVResult SendFrame(TSharedPtr<FVideoResourceD3D11> const& Resource, uint32 Timestamp, bool bForceKeyframe = false) override;
-	virtual FAVResult ReceivePacket(FVideoPacket& OutPacket) override;
+	NVENC_API virtual FAVResult Open(TSharedRef<FAVDevice> const& NewDevice, TSharedRef<FAVInstance> const& NewInstance) override;
+	NVENC_API virtual void Close() override;
+	NVENC_API virtual bool IsOpen() const override;
+	NVENC_API virtual FAVResult ApplyConfig() override;
+	NVENC_API virtual FAVResult SendFrame(TSharedPtr<FVideoResourceD3D11> const& Resource, uint32 Timestamp, bool bForceKeyframe = false) override;
+	NVENC_API virtual FAVResult ReceivePacket(FVideoPacket& OutPacket) override;
 	// End TVideoEncoder interface
 
 	// Begin FEncoderNVENC interface
-	virtual	bool IsInitialized() const override;
-	virtual FAVResult OnPostInitialization(FVideoEncoderConfigNVENC& PendingConfig) override;
+	NVENC_API virtual	bool IsInitialized() const override;
+	NVENC_API virtual FAVResult OnPostInitialization(FVideoEncoderConfigNVENC& PendingConfig) override;
 	// End FEncoderNVENC interface
 };
 
 /**
  *  An AVCodecs VideoEncoder that takes a FVideoResourceD3D12.
  **/
-class NVENC_API FVideoEncoderNVENCD3D12 : public TVideoEncoder<FVideoResourceD3D12, FVideoEncoderConfigNVENC>, public FEncoderNVENC
+class FVideoEncoderNVENCD3D12 : public TVideoEncoder<FVideoResourceD3D12, FVideoEncoderConfigNVENC>, public FEncoderNVENC
 {
 private:
 	TRefCountPtr<ID3D12Fence> InputFence;
@@ -145,20 +145,20 @@ private:
 
 public:
 	FVideoEncoderNVENCD3D12() = default;
-	virtual ~FVideoEncoderNVENCD3D12() override;
+	NVENC_API virtual ~FVideoEncoderNVENCD3D12() override;
 
 	// Begin TVideoEncoder interface
-	virtual FAVResult Open(TSharedRef<FAVDevice> const& NewDevice, TSharedRef<FAVInstance> const& NewInstance) override;
-	virtual void Close() override;
-	virtual bool IsOpen() const override;
-	virtual FAVResult ApplyConfig() override;
-	virtual FAVResult SendFrame(TSharedPtr<FVideoResourceD3D12> const& Resource, uint32 Timestamp, bool bForceKeyframe = false) override;
-	virtual FAVResult ReceivePacket(FVideoPacket& OutPacket) override;
+	NVENC_API virtual FAVResult Open(TSharedRef<FAVDevice> const& NewDevice, TSharedRef<FAVInstance> const& NewInstance) override;
+	NVENC_API virtual void Close() override;
+	NVENC_API virtual bool IsOpen() const override;
+	NVENC_API virtual FAVResult ApplyConfig() override;
+	NVENC_API virtual FAVResult SendFrame(TSharedPtr<FVideoResourceD3D12> const& Resource, uint32 Timestamp, bool bForceKeyframe = false) override;
+	NVENC_API virtual FAVResult ReceivePacket(FVideoPacket& OutPacket) override;
 	// End TVideoEncoder interface
 
 	// Begin FEncoderNVENC interface
-	virtual	bool IsInitialized() const override;
-	virtual FAVResult OnPostInitialization(FVideoEncoderConfigNVENC& PendingConfig) override;
+	NVENC_API virtual	bool IsInitialized() const override;
+	NVENC_API virtual FAVResult OnPostInitialization(FVideoEncoderConfigNVENC& PendingConfig) override;
 	// End FEncoderNVENC interface
 };
-#endif // PLATFORM_WINDOWS
\ No newline at end of file
+#endif // PLATFORM_WINDOWS
diff --git a/Engine/Plugins/Experimental/WebSocketNetworking/Source/WebSocketNetworking/Private/WebSocketNetDriver.cpp b/Engine/Plugins/Experimental/WebSocketNetworking/Source/WebSocketNetworking/Private/WebSocketNetDriver.cpp
--- a/Engine/Plugins/Experimental/WebSocketNetworking/Source/WebSocketNetworking/Private/WebSocketNetDriver.cpp
+++ b/Engine/Plugins/Experimental/WebSocketNetworking/Source/WebSocketNetworking/Private/WebSocketNetDriver.cpp
@@ -15,6 +15,8 @@ Unreal WebSocket network driver.
 #include "Misc/CommandLine.h"
 #include "PacketHandler.h"
 
+#include UE_INLINE_GENERATED_CPP_BY_NAME(WebSocketNetDriver)
+
 /*-----------------------------------------------------------------------------
 Declarations.
 -----------------------------------------------------------------------------*/
diff --git a/Engine/Plugins/Experimental/WebSocketNetworking/Source/WebSocketNetworking/Public/INetworkingWebSocket.h b/Engine/Plugins/Experimental/WebSocketNetworking/Source/WebSocketNetworking/Public/INetworkingWebSocket.h
--- a/Engine/Plugins/Experimental/WebSocketNetworking/Source/WebSocketNetworking/Public/INetworkingWebSocket.h
+++ b/Engine/Plugins/Experimental/WebSocketNetworking/Source/WebSocketNetworking/Public/INetworkingWebSocket.h
@@ -5,7 +5,7 @@
 #include "CoreFwd.h"
 #include "WebSocketNetworkingDelegates.h"
 
-class WEBSOCKETNETWORKING_API INetworkingWebSocket
+class INetworkingWebSocket
 {
 public:
 	virtual ~INetworkingWebSocket() {}
@@ -33,4 +33,4 @@ public:
 	virtual struct sockaddr_in* GetRemoteAddr() = 0;
 	virtual TArray<FString> GetUrlArgs() = 0;
 	virtual FString GetUrlArgByName(const FString& Name) = 0;
-};
\ No newline at end of file
+};
diff --git a/Engine/Plugins/Experimental/WebSocketNetworking/Source/WebSocketNetworking/Public/IWebSocketNetworkingModule.h b/Engine/Plugins/Experimental/WebSocketNetworking/Source/WebSocketNetworking/Public/IWebSocketNetworkingModule.h
--- a/Engine/Plugins/Experimental/WebSocketNetworking/Source/WebSocketNetworking/Public/IWebSocketNetworkingModule.h
+++ b/Engine/Plugins/Experimental/WebSocketNetworking/Source/WebSocketNetworking/Public/IWebSocketNetworkingModule.h
@@ -13,7 +13,7 @@ class INetworkingWebSocket;
 /* 
  * Interface for the WebSocketNetworking module. 
  */
-class WEBSOCKETNETWORKING_API IWebSocketNetworkingModule
+class IWebSocketNetworkingModule
 	: public IModuleInterface
 {
 public:
diff --git a/Engine/Plugins/Experimental/WebSocketNetworking/Source/WebSocketNetworking/Public/IWebSocketServer.h b/Engine/Plugins/Experimental/WebSocketNetworking/Source/WebSocketNetworking/Public/IWebSocketServer.h
--- a/Engine/Plugins/Experimental/WebSocketNetworking/Source/WebSocketNetworking/Public/IWebSocketServer.h
+++ b/Engine/Plugins/Experimental/WebSocketNetworking/Source/WebSocketNetworking/Public/IWebSocketServer.h
@@ -11,7 +11,7 @@
  *
  * For example we may wish to map `C:/www/Public` to `/` of our webserver.
  **/
-class WEBSOCKETNETWORKING_API FWebSocketHttpMount
+class FWebSocketHttpMount
 {
 public:
 	/**
@@ -50,7 +50,7 @@ private:
 	std::string DefaultFile = "index.html";
 };
 
-class WEBSOCKETNETWORKING_API FWebSocketServerCertificates
+class FWebSocketServerCertificates
 {
 public:
 	/**
@@ -78,7 +78,7 @@ private:
 	std::string PrivateKeyFilePath;
 };
 
-class WEBSOCKETNETWORKING_API IWebSocketServer
+class IWebSocketServer
 {
 public:
 	virtual ~IWebSocketServer() {}
@@ -113,4 +113,4 @@ public:
 
 	/** Describe this libwebsocket server */
 	virtual FString Info() = 0;
-};
\ No newline at end of file
+};
diff --git a/Engine/Plugins/Experimental/WebSocketNetworking/Source/WebSocketNetworking/Public/WebSocketConnection.h b/Engine/Plugins/Experimental/WebSocketNetworking/Source/WebSocketNetworking/Public/WebSocketConnection.h
--- a/Engine/Plugins/Experimental/WebSocketNetworking/Source/WebSocketNetworking/Public/WebSocketConnection.h
+++ b/Engine/Plugins/Experimental/WebSocketNetworking/Source/WebSocketNetworking/Public/WebSocketConnection.h
@@ -5,29 +5,33 @@
 #include "Engine/NetConnection.h"
 #include "WebSocketConnection.generated.h"
 
+#define UE_API WEBSOCKETNETWORKING_API
 
-UCLASS(transient, config = Engine)
-class WEBSOCKETNETWORKING_API UWebSocketConnection : public UNetConnection
+
+UCLASS(MinimalAPI, transient, config = Engine)
+class UWebSocketConnection : public UNetConnection
 {
 	GENERATED_UCLASS_BODY()
 
 	class INetworkingWebSocket* WebSocket;
 
 	//~ Begin NetConnection Interface
-	virtual void InitBase(UNetDriver* InDriver, class FSocket* InSocket, const FURL& InURL, EConnectionState InState, int32 InMaxPacket = 0, int32 InPacketOverhead = 0) override;
-	virtual void InitRemoteConnection(UNetDriver* InDriver, class FSocket* InSocket, const FURL& InURL, const class FInternetAddr& InRemoteAddr, EConnectionState InState, int32 InMaxPacket = 0, int32 InPacketOverhead = 0) override;
-	virtual void InitLocalConnection(UNetDriver* InDriver, class FSocket* InSocket, const FURL& InURL, EConnectionState InState, int32 InMaxPacket = 0, int32 InPacketOverhead = 0) override;
-	virtual void LowLevelSend(void* Data, int32 CountBits, FOutPacketTraits& Traits) override;
-	FString LowLevelGetRemoteAddress(bool bAppendPort = false) override;
-	FString LowLevelDescribe() override;
-	virtual void Tick(float DeltaSeconds) override;
-	virtual void FinishDestroy();
-	virtual void ReceivedRawPacket(void* Data,int32 Count);
+	UE_API virtual void InitBase(UNetDriver* InDriver, class FSocket* InSocket, const FURL& InURL, EConnectionState InState, int32 InMaxPacket = 0, int32 InPacketOverhead = 0) override;
+	UE_API virtual void InitRemoteConnection(UNetDriver* InDriver, class FSocket* InSocket, const FURL& InURL, const class FInternetAddr& InRemoteAddr, EConnectionState InState, int32 InMaxPacket = 0, int32 InPacketOverhead = 0) override;
+	UE_API virtual void InitLocalConnection(UNetDriver* InDriver, class FSocket* InSocket, const FURL& InURL, EConnectionState InState, int32 InMaxPacket = 0, int32 InPacketOverhead = 0) override;
+	UE_API virtual void LowLevelSend(void* Data, int32 CountBits, FOutPacketTraits& Traits) override;
+	UE_API FString LowLevelGetRemoteAddress(bool bAppendPort = false) override;
+	UE_API FString LowLevelDescribe() override;
+	UE_API virtual void Tick(float DeltaSeconds) override;
+	UE_API virtual void FinishDestroy();
+	UE_API virtual void ReceivedRawPacket(void* Data,int32 Count);
 	//~ End NetConnection Interface
 
 
-	void SetWebSocket(INetworkingWebSocket* InWebSocket);
-	INetworkingWebSocket* GetWebSocket();
+	UE_API void SetWebSocket(INetworkingWebSocket* InWebSocket);
+	UE_API INetworkingWebSocket* GetWebSocket();
 
 	bool bChallengeHandshake = false;
 };
+
+#undef UE_API
diff --git a/Engine/Plugins/Experimental/WebSocketNetworking/Source/WebSocketNetworking/Public/WebSocketNetDriver.h b/Engine/Plugins/Experimental/WebSocketNetworking/Source/WebSocketNetworking/Public/WebSocketNetDriver.h
--- a/Engine/Plugins/Experimental/WebSocketNetworking/Source/WebSocketNetworking/Public/WebSocketNetDriver.h
+++ b/Engine/Plugins/Experimental/WebSocketNetworking/Source/WebSocketNetworking/Public/WebSocketNetDriver.h
@@ -9,11 +9,13 @@
 #include "Engine/NetDriver.h"
 #include "WebSocketNetDriver.generated.h"
 
+#define UE_API WEBSOCKETNETWORKING_API
+
 class INetworkingWebSocket;
 class IWebSocketServer;
 
-UCLASS(transient, config = Engine)
-class WEBSOCKETNETWORKING_API UWebSocketNetDriver : public UNetDriver
+UCLASS(MinimalAPI, transient, config = Engine)
+class UWebSocketNetDriver : public UNetDriver
 {
 	GENERATED_UCLASS_BODY()
 
@@ -22,34 +24,34 @@ class WEBSOCKETNETWORKING_API UWebSocketNetDriver : public UNetDriver
 	int32 WebSocketPort;
 
 	//~ Begin UNetDriver Interface.
-	virtual bool IsAvailable() const override;
-	virtual bool InitBase(bool bInitAsClient, FNetworkNotify* InNotify, const FURL& URL, bool bReuseAddressAndPort, FString& Error) override;
-	virtual bool InitConnect(FNetworkNotify* InNotify, const FURL& ConnectURL, FString& Error) override;
-	virtual bool InitListen(FNetworkNotify* InNotify, FURL& LocalURL, bool bReuseAddressAndPort, FString& Error) override;
-	virtual void TickDispatch(float DeltaTime) override;
-	virtual void LowLevelSend(TSharedPtr<const FInternetAddr> Address, void* Data, int32 CountBits, FOutPacketTraits& Traits) override;
-	virtual FString LowLevelGetNetworkNumber() override;
-	virtual void LowLevelDestroy() override;
-	virtual bool IsNetResourceValid(void) override;
+	UE_API virtual bool IsAvailable() const override;
+	UE_API virtual bool InitBase(bool bInitAsClient, FNetworkNotify* InNotify, const FURL& URL, bool bReuseAddressAndPort, FString& Error) override;
+	UE_API virtual bool InitConnect(FNetworkNotify* InNotify, const FURL& ConnectURL, FString& Error) override;
+	UE_API virtual bool InitListen(FNetworkNotify* InNotify, FURL& LocalURL, bool bReuseAddressAndPort, FString& Error) override;
+	UE_API virtual void TickDispatch(float DeltaTime) override;
+	UE_API virtual void LowLevelSend(TSharedPtr<const FInternetAddr> Address, void* Data, int32 CountBits, FOutPacketTraits& Traits) override;
+	UE_API virtual FString LowLevelGetNetworkNumber() override;
+	UE_API virtual void LowLevelDestroy() override;
+	UE_API virtual bool IsNetResourceValid(void) override;
 
 	// stub implementation because for websockets we don't use any underlying socket sub system.
-	virtual class ISocketSubsystem* GetSocketSubsystem() override;
+	UE_API virtual class ISocketSubsystem* GetSocketSubsystem() override;
 
 	//~ End UNetDriver Interface.
 
 	//~ Begin FExec Interface
 #if UE_ALLOW_EXEC_COMMANDS
-	virtual bool Exec(UWorld* InWorld, const TCHAR* Cmd, FOutputDevice& Ar = *GLog) override;
+	UE_API virtual bool Exec(UWorld* InWorld, const TCHAR* Cmd, FOutputDevice& Ar = *GLog) override;
 #endif
 	//~ End FExec Interface
 
 	/**
 	* Exec command handlers
 	*/
-	bool HandleSocketsCommand(const TCHAR* Cmd, FOutputDevice& Ar, UWorld* InWorld);
+	UE_API bool HandleSocketsCommand(const TCHAR* Cmd, FOutputDevice& Ar, UWorld* InWorld);
 
 	/** @return connection to server */
-	class UWebSocketConnection* GetServerConnection();
+	UE_API class UWebSocketConnection* GetServerConnection();
 
 	/************************************************************************/
 	/* IWebSocketServer														*/
@@ -61,14 +63,16 @@ class WEBSOCKETNETWORKING_API UWebSocketNetDriver : public UNetDriver
 	/* Callback Function for New Connection from a client is accepted by this server   */
 	/******************************************************************************/
 
-	void OnWebSocketClientConnected(INetworkingWebSocket*); // to the server.
+	UE_API void OnWebSocketClientConnected(INetworkingWebSocket*); // to the server.
 
 
 	/************************************************************************/
 	/* Callback Function for when this client Connects to the server				*/
 	/************************************************************************/
 
-	void OnWebSocketServerConnected(); // to the client.
+	UE_API void OnWebSocketServerConnected(); // to the client.
 
 
 };
+
+#undef UE_API
diff --git a/Engine/Plugins/Media/PixelCapture/Source/PixelCapture/Internal/PixelCaptureCapturerMediaCapture.h b/Engine/Plugins/Media/PixelCapture/Source/PixelCapture/Internal/PixelCaptureCapturerMediaCapture.h
--- a/Engine/Plugins/Media/PixelCapture/Source/PixelCapture/Internal/PixelCaptureCapturerMediaCapture.h
+++ b/Engine/Plugins/Media/PixelCapture/Source/PixelCapture/Internal/PixelCaptureCapturerMediaCapture.h
@@ -10,8 +10,10 @@
 
 #include "PixelCaptureCapturerMediaCapture.generated.h"
 
-UCLASS(BlueprintType)
-class PIXELCAPTURE_API UPixelCaptureMediaOuput : public UMediaOutput
+#define UE_API PIXELCAPTURE_API
+
+UCLASS(MinimalAPI, BlueprintType)
+class UPixelCaptureMediaOuput : public UMediaOutput
 {
 	GENERATED_BODY()
 
@@ -25,8 +27,8 @@ private:
 	FIntPoint RequestedSize;
 };
 
-UCLASS(BlueprintType)
-class PIXELCAPTURE_API UPixelCaptureMediaCapture : public UMediaCapture
+UCLASS(MinimalAPI, BlueprintType)
+class UPixelCaptureMediaCapture : public UMediaCapture
 {
 	GENERATED_BODY()
 
@@ -61,23 +63,23 @@ protected:
 		return true;
 	}
 
-	virtual void OnRHIResourceCaptured_AnyThread(
+	UE_API virtual void OnRHIResourceCaptured_AnyThread(
 		const FCaptureBaseData&								   BaseData,
 		TSharedPtr<FMediaCaptureUserData, ESPMode::ThreadSafe> UserData,
 		FTextureRHIRef										   Texture) override;
 
-	virtual void OnRHIResourceCaptured_RenderingThread(
+	UE_API virtual void OnRHIResourceCaptured_RenderingThread(
 		FRHICommandListImmediate& RHICmdList,
 		const FCaptureBaseData& BaseData,
 		TSharedPtr<FMediaCaptureUserData, ESPMode::ThreadSafe> UserData,
 		FTextureRHIRef Texture) override;
 
-	virtual void OnFrameCaptured_AnyThread(
+	UE_API virtual void OnFrameCaptured_AnyThread(
 		const FCaptureBaseData&								   BaseData,
 		TSharedPtr<FMediaCaptureUserData, ESPMode::ThreadSafe> UserData,
 		const FMediaCaptureResourceData&					   ResourceData) override;
 
-	virtual void OnCustomCapture_RenderingThread(
+	UE_API virtual void OnCustomCapture_RenderingThread(
 		FRDGBuilder&										   GraphBuilder,
 		const FCaptureBaseData&								   InBaseData,
 		TSharedPtr<FMediaCaptureUserData, ESPMode::ThreadSafe> InUserData,
@@ -87,7 +89,7 @@ protected:
 		FVector2D											   CropU,
 		FVector2D											   CropV) override;
 
-	virtual void OnFrameCaptured_RenderingThread(
+	UE_API virtual void OnFrameCaptured_RenderingThread(
 		const FCaptureBaseData&								   InBaseData,
 		TSharedPtr<FMediaCaptureUserData, ESPMode::ThreadSafe> InUserData,
 		void*												   InBuffer,
@@ -112,8 +114,8 @@ protected:
 #endif // PLATFORM_MAC
 	}
 
-	virtual ETextureCreateFlags GetOutputTextureFlags() const override;
-	virtual void				WaitForGPU(FRHITexture* InRHITexture) override;
+	UE_API virtual ETextureCreateFlags GetOutputTextureFlags() const override;
+	UE_API virtual void				WaitForGPU(FRHITexture* InRHITexture) override;
 
 private:
 	int32 Format = PixelCaptureBufferFormat::FORMAT_UNKNOWN;
@@ -130,32 +132,34 @@ private:
  * Input: FPixelCaptureInputFrameRHI
  * Output: FPixelCaptureOutputFrameRHI/FPixelCaptureOutputFrameI420
  */
-class PIXELCAPTURE_API FPixelCaptureCapturerMediaCapture : public FPixelCaptureCapturer, public TSharedFromThis<FPixelCaptureCapturerMediaCapture>
+class FPixelCaptureCapturerMediaCapture : public FPixelCaptureCapturer, public TSharedFromThis<FPixelCaptureCapturerMediaCapture>
 {
 public:
 	/**
 	 * Creates a new Capturer capturing the input frame at the given resolution.
 	 * @param InOutputResolution The resolution of the resulting output capture.
 	 */
-	static TSharedPtr<FPixelCaptureCapturerMediaCapture> Create(FPixelCaptureCapturerConfig Config, int32 InFormat);
-	virtual ~FPixelCaptureCapturerMediaCapture() override;
+	static UE_API TSharedPtr<FPixelCaptureCapturerMediaCapture> Create(FPixelCaptureCapturerConfig Config, int32 InFormat);
+	UE_API virtual ~FPixelCaptureCapturerMediaCapture() override;
 
 protected:
 	virtual FString					  GetCapturerName() const override { return "MediaCapture Copy"; }
-	virtual IPixelCaptureOutputFrame* CreateOutputBuffer(int32 InputWidth, int32 InputHeight) override;
-	virtual void					  BeginProcess(const IPixelCaptureInputFrame& InputFrame, TSharedPtr<IPixelCaptureOutputFrame> OutputBuffer) override;
+	UE_API virtual IPixelCaptureOutputFrame* CreateOutputBuffer(int32 InputWidth, int32 InputHeight) override;
+	UE_API virtual void					  BeginProcess(const IPixelCaptureInputFrame& InputFrame, TSharedPtr<IPixelCaptureOutputFrame> OutputBuffer) override;
 
 private:
 	uint64 LastFrameCounterRenderThread = 0;
 	int32  Format = PixelCaptureBufferFormat::FORMAT_UNKNOWN;
 	bool   bMediaCaptureInitialized = false;
 
-	FPixelCaptureCapturerMediaCapture(FPixelCaptureCapturerConfig& Config, int32 InFormat);
-	void InitializeMediaCapture();
+	UE_API FPixelCaptureCapturerMediaCapture(FPixelCaptureCapturerConfig& Config, int32 InFormat);
+	UE_API void InitializeMediaCapture();
 
-	void InvalidateOutputBuffer(TSharedPtr<IPixelCaptureOutputFrame> OutputBuffer);
+	UE_API void InvalidateOutputBuffer(TSharedPtr<IPixelCaptureOutputFrame> OutputBuffer);
 
 	TObjectPtr<UPixelCaptureMediaCapture> MediaCapture = nullptr;
 
 	TObjectPtr<UPixelCaptureMediaOuput> MediaOutput = nullptr;
 };
+
+#undef UE_API
diff --git a/Engine/Plugins/Media/PixelCapture/Source/PixelCapture/Private/PixelCaptureCapturerMediaCapture.cpp b/Engine/Plugins/Media/PixelCapture/Source/PixelCapture/Private/PixelCaptureCapturerMediaCapture.cpp
--- a/Engine/Plugins/Media/PixelCapture/Source/PixelCapture/Private/PixelCaptureCapturerMediaCapture.cpp
+++ b/Engine/Plugins/Media/PixelCapture/Source/PixelCapture/Private/PixelCaptureCapturerMediaCapture.cpp
@@ -12,6 +12,8 @@
 
 #include "libyuv/convert.h"
 
+#include UE_INLINE_GENERATED_CPP_BY_NAME(PixelCaptureCapturerMediaCapture)
+
 void UPixelCaptureMediaCapture::OnRHIResourceCaptured_AnyThread(
 	const FCaptureBaseData&								   BaseData,
 	TSharedPtr<FMediaCaptureUserData, ESPMode::ThreadSafe> UserData,
@@ -403,4 +405,4 @@ void FPixelCaptureCapturerMediaCapture::InvalidateOutputBuffer(TSharedPtr<IPixel
 	{
 		UE_LOGFMT(LogPixelCapture, Error, "FPixelCaptureCapturerMediaCapture: Invalid pixel format. Expected either FORMAT_RHI or FORMAT_I420");
 	}
-}
\ No newline at end of file
+}
diff --git a/Engine/Plugins/Media/PixelCapture/Source/PixelCapture/Public/IPixelCaptureBuffer.h b/Engine/Plugins/Media/PixelCapture/Source/PixelCapture/Public/IPixelCaptureBuffer.h
--- a/Engine/Plugins/Media/PixelCapture/Source/PixelCapture/Public/IPixelCaptureBuffer.h
+++ b/Engine/Plugins/Media/PixelCapture/Source/PixelCapture/Public/IPixelCaptureBuffer.h
@@ -8,7 +8,7 @@
 /**
  * Buffer Template for defining frame data.
  */
-class PIXELCAPTURE_API IPixelCaptureBuffer
+class IPixelCaptureBuffer
 {
 public:
 	/**
@@ -40,4 +40,4 @@ public:
 	 * Gets a pointer to the beginning of the entire buffer for editing.
 	 */
 	virtual uint8_t* GetMutableData() = 0;
-};
\ No newline at end of file
+};
diff --git a/Engine/Plugins/Media/PixelCapture/Source/PixelCapture/Public/IPixelCaptureCapturerSource.h b/Engine/Plugins/Media/PixelCapture/Source/PixelCapture/Public/IPixelCaptureCapturerSource.h
--- a/Engine/Plugins/Media/PixelCapture/Source/PixelCapture/Public/IPixelCaptureCapturerSource.h
+++ b/Engine/Plugins/Media/PixelCapture/Source/PixelCapture/Public/IPixelCaptureCapturerSource.h
@@ -8,7 +8,7 @@
  * The interface that is queried when the Layered or Multi capturers wish to create specific
  * Capturers for a given format.
  */
-class PIXELCAPTURE_API IPixelCaptureCapturerSource
+class IPixelCaptureCapturerSource
 {
 public:
 	/**
diff --git a/Engine/Plugins/Media/PixelCapture/Source/PixelCapture/Public/IPixelCaptureInputFrame.h b/Engine/Plugins/Media/PixelCapture/Source/PixelCapture/Public/IPixelCaptureInputFrame.h
--- a/Engine/Plugins/Media/PixelCapture/Source/PixelCapture/Public/IPixelCaptureInputFrame.h
+++ b/Engine/Plugins/Media/PixelCapture/Source/PixelCapture/Public/IPixelCaptureInputFrame.h
@@ -11,7 +11,7 @@
  * FPixelCaptureCapturer objects for each format/resolution. The capturer should then cast
  * to the known type fed into the system to unwrap the data.
  */
-class PIXELCAPTURE_API IPixelCaptureInputFrame
+class IPixelCaptureInputFrame
 {
 public:
 	IPixelCaptureInputFrame() { Metadata.SourceTime = rtc::TimeMillis(); }
diff --git a/Engine/Plugins/Media/PixelCapture/Source/PixelCapture/Public/IPixelCaptureOutputFrame.h b/Engine/Plugins/Media/PixelCapture/Source/PixelCapture/Public/IPixelCaptureOutputFrame.h
--- a/Engine/Plugins/Media/PixelCapture/Source/PixelCapture/Public/IPixelCaptureOutputFrame.h
+++ b/Engine/Plugins/Media/PixelCapture/Source/PixelCapture/Public/IPixelCaptureOutputFrame.h
@@ -9,7 +9,7 @@
  * own result types. You must implement GetWidth and GetHeight to return the width
  * and height of the frame. Add your own method to extract the captured data.
  */
-class PIXELCAPTURE_API IPixelCaptureOutputFrame
+class IPixelCaptureOutputFrame
 {
 public:
 	virtual ~IPixelCaptureOutputFrame() = default;
diff --git a/Engine/Plugins/Media/PixelCapture/Source/PixelCapture/Public/PixelCaptureBufferI010.h b/Engine/Plugins/Media/PixelCapture/Source/PixelCapture/Public/PixelCaptureBufferI010.h
--- a/Engine/Plugins/Media/PixelCapture/Source/PixelCapture/Public/PixelCaptureBufferI010.h
+++ b/Engine/Plugins/Media/PixelCapture/Source/PixelCapture/Public/PixelCaptureBufferI010.h
@@ -4,16 +4,18 @@
 
 #include "IPixelCaptureBuffer.h"
 
+#define UE_API PIXELCAPTURE_API
+
 /**
  * Buffer class for holding I010 frame data.
  */
-class PIXELCAPTURE_API FPixelCaptureBufferI010 : public IPixelCaptureBuffer
+class FPixelCaptureBufferI010 : public IPixelCaptureBuffer
 {
 public:
 	/**
 	 * Create an empty buffer of the specified size.
 	 */
-	FPixelCaptureBufferI010(int InWidth, int InHeight);
+	UE_API FPixelCaptureBufferI010(int InWidth, int InHeight);
 
 	FPixelCaptureBufferI010(const FPixelCaptureBufferI010& Other) = default;
 	virtual ~FPixelCaptureBufferI010() = default;
@@ -46,52 +48,52 @@ public:
 	/**
 	 * Gets a const pointer to the beginning of entire buffer for reading.
 	 */
-	virtual const uint8_t* GetData() const override;
+	UE_API virtual const uint8_t* GetData() const override;
 
 	/**
 	 * Gets a const pointer to the beginning of the Y plane for reading.
 	 */
-	const uint16_t* GetDataY() const;
+	UE_API const uint16_t* GetDataY() const;
 
 	/**
 	 * Gets a const pointer to the beginning of the U plane for reading.
 	 */
-	const uint16_t* GetDataU() const;
+	UE_API const uint16_t* GetDataU() const;
 
 	/**
 	 * Gets a const pointer to the beginning of the V plane for reading.
 	 */
-	const uint16_t* GetDataV() const;
+	UE_API const uint16_t* GetDataV() const;
 
 	/**
 	 * Gets a pointer to the beginning of the entire buffer for editing.
 	 */
-	virtual uint8_t* GetMutableData() override;
+	UE_API virtual uint8_t* GetMutableData() override;
 
 	/**
 	 * Gets a pointer to the beginning of the Y plane for editing.
 	 */
-	uint16_t* GetMutableDataY();
+	UE_API uint16_t* GetMutableDataY();
 
 	/**
 	 * Gets a pointer to the beginning of the U plane for editing.
 	 */
-	uint16_t* GetMutableDataU();
+	UE_API uint16_t* GetMutableDataU();
 
 	/**
 	 * Gets a pointer to the beginning of the V plane for editing.
 	 */
-	uint16_t* GetMutableDataV();
+	UE_API uint16_t* GetMutableDataV();
 
 	/**
 	 * Gets the size of the Y plane in bytes.
 	 */
-	int32 GetDataSizeY() const;
+	UE_API int32 GetDataSizeY() const;
 
 	/**
 	 * Gets the size of both the U and V planes in bytes.
 	 */
-	int32 GetDataSizeUV() const;
+	UE_API int32 GetDataSizeUV() const;
 
 	/**
 	 * Gets the total size of this buffer.
@@ -113,3 +115,5 @@ private:
 
 	TArray<uint16_t> Data;
 };
+
+#undef UE_API
diff --git a/Engine/Plugins/Media/PixelCapture/Source/PixelCapture/Public/PixelCaptureBufferI420.h b/Engine/Plugins/Media/PixelCapture/Source/PixelCapture/Public/PixelCaptureBufferI420.h
--- a/Engine/Plugins/Media/PixelCapture/Source/PixelCapture/Public/PixelCaptureBufferI420.h
+++ b/Engine/Plugins/Media/PixelCapture/Source/PixelCapture/Public/PixelCaptureBufferI420.h
@@ -4,16 +4,18 @@
 
 #include "IPixelCaptureBuffer.h"
 
+#define UE_API PIXELCAPTURE_API
+
 /**
  * Buffer class for holding I420 frame data.
  */
-class PIXELCAPTURE_API FPixelCaptureBufferI420 : public IPixelCaptureBuffer
+class FPixelCaptureBufferI420 : public IPixelCaptureBuffer
 {
 public:
 	/**
 	 * Create an empty buffer of the specified size.
 	 */
-	FPixelCaptureBufferI420(int InWidth, int InHeight);
+	UE_API FPixelCaptureBufferI420(int InWidth, int InHeight);
 
 	FPixelCaptureBufferI420(const FPixelCaptureBufferI420& Other) = default;
 	virtual	~FPixelCaptureBufferI420() = default;
@@ -46,52 +48,52 @@ public:
 	/**
 	 * Gets a const pointer to the beginning of entire buffer for reading.
 	 */
-	virtual const uint8_t* GetData() const override;
+	UE_API virtual const uint8_t* GetData() const override;
 
 	/**
 	 * Gets a const pointer to the beginning of the Y plane for reading.
 	 */
-	const uint8_t* GetDataY() const;
+	UE_API const uint8_t* GetDataY() const;
 
 	/**
 	 * Gets a const pointer to the beginning of the U plane for reading.
 	 */
-	const uint8_t* GetDataU() const;
+	UE_API const uint8_t* GetDataU() const;
 
 	/**
 	 * Gets a const pointer to the beginning of the V plane for reading.
 	 */
-	const uint8_t* GetDataV() const;
+	UE_API const uint8_t* GetDataV() const;
 
 	/**
 	 * Gets a pointer to the beginning of the entire buffer for editing.
 	 */
-	virtual uint8_t*  GetMutableData() override;
+	UE_API virtual uint8_t*  GetMutableData() override;
 
 	/**
 	 * Gets a pointer to the beginning of the Y plane for editing.
 	 */
-	uint8_t* GetMutableDataY();
+	UE_API uint8_t* GetMutableDataY();
 
 	/**
 	 * Gets a pointer to the beginning of the U plane for editing.
 	 */
-	uint8_t* GetMutableDataU();
+	UE_API uint8_t* GetMutableDataU();
 
 	/**
 	 * Gets a pointer to the beginning of the V plane for editing.
 	 */
-	uint8_t* GetMutableDataV();
+	UE_API uint8_t* GetMutableDataV();
 
 	/**
 	 * Gets the size of the Y plane in bytes.
 	 */
-	int32 GetDataSizeY() const;
+	UE_API int32 GetDataSizeY() const;
 
 	/**
 	 * Gets the size of both the U and V planes in bytes.
 	 */
-	int32 GetDataSizeUV() const;
+	UE_API int32 GetDataSizeUV() const;
 
 	/**
 	 * Gets the total size of this buffer.
@@ -113,3 +115,5 @@ private:
 
 	TArray<uint8_t> Data;
 };
+
+#undef UE_API
diff --git a/Engine/Plugins/Media/PixelCapture/Source/PixelCapture/Public/PixelCaptureBufferNV12.h b/Engine/Plugins/Media/PixelCapture/Source/PixelCapture/Public/PixelCaptureBufferNV12.h
--- a/Engine/Plugins/Media/PixelCapture/Source/PixelCapture/Public/PixelCaptureBufferNV12.h
+++ b/Engine/Plugins/Media/PixelCapture/Source/PixelCapture/Public/PixelCaptureBufferNV12.h
@@ -4,16 +4,18 @@
 
 #include "IPixelCaptureBuffer.h"
 
+#define UE_API PIXELCAPTURE_API
+
 /**
  * Buffer class for holding I420 frame data.
  */
-class PIXELCAPTURE_API FPixelCaptureBufferNV12  : public IPixelCaptureBuffer
+class FPixelCaptureBufferNV12  : public IPixelCaptureBuffer
 {
 public:
 	/**
 	 * Create an empty buffer of the specified size.
 	 */
-	FPixelCaptureBufferNV12(int InWidth, int InHeight);
+	UE_API FPixelCaptureBufferNV12(int InWidth, int InHeight);
 
 	FPixelCaptureBufferNV12(const FPixelCaptureBufferNV12& Other) = default;
 	virtual	~FPixelCaptureBufferNV12() = default;
@@ -46,42 +48,42 @@ public:
 	/**
 	 * Gets a const pointer to the beginning of entire buffer for reading.
 	 */
-	virtual const uint8_t* GetData() const override;
+	UE_API virtual const uint8_t* GetData() const override;
 
 	/**
 	 * Gets a const pointer to the beginning of the Y plane for reading.
 	 */
-	const uint8_t* GetDataY() const;
+	UE_API const uint8_t* GetDataY() const;
 
 	/**
 	 * Gets a const pointer to the beginning of the U plane for reading.
 	 */
-	const uint8_t* GetDataUV() const;
+	UE_API const uint8_t* GetDataUV() const;
 
 	/**
 	 * Gets a pointer to the beginning of the entire buffer for editing.
 	 */
-	virtual uint8_t* GetMutableData() override;
+	UE_API virtual uint8_t* GetMutableData() override;
 
 	/**
 	 * Gets a pointer to the beginning of the Y plane for editing.
 	 */
-	uint8_t* GetMutableDataY();
+	UE_API uint8_t* GetMutableDataY();
 
 	/**
 	 * Gets a pointer to the beginning of the U plane for editing.
 	 */
-	uint8_t* GetMutableDataUV();
+	UE_API uint8_t* GetMutableDataUV();
 
 	/**
 	 * Gets the size of the Y plane in bytes.
 	 */
-	int GetDataSizeY() const;
+	UE_API int GetDataSizeY() const;
 
 	/**
 	 * Gets the size of both the UV plane in bytes.
 	 */
-	int GetDataSizeUV() const;
+	UE_API int GetDataSizeUV() const;
 
 	/**
 	 * Gets the total size of this buffer.
@@ -103,3 +105,5 @@ private:
 
 	TArray<uint8_t> Data;
 };
+
+#undef UE_API
diff --git a/Engine/Plugins/Media/PixelCapture/Source/PixelCapture/Public/PixelCaptureCapturer.h b/Engine/Plugins/Media/PixelCapture/Source/PixelCapture/Public/PixelCaptureCapturer.h
--- a/Engine/Plugins/Media/PixelCapture/Source/PixelCapture/Public/PixelCaptureCapturer.h
+++ b/Engine/Plugins/Media/PixelCapture/Source/PixelCapture/Public/PixelCaptureCapturer.h
@@ -6,6 +6,8 @@
 #include "IPixelCaptureOutputFrame.h"
 #include "HAL/ThreadSafeCounter64.h"
 
+#define UE_API PIXELCAPTURE_API
+
 namespace UE::PixelCapture
 {
 	class FOutputFrameBuffer;
@@ -16,7 +18,7 @@ namespace UE::PixelCapture
 	PIXELCAPTURE_API void MarkGPUWorkEnd(TSharedPtr<IPixelCaptureOutputFrame> OutputBuffer);
 } // namespace UE::PixelCapture
 
-struct PIXELCAPTURE_API FPixelCaptureCapturerConfig
+struct FPixelCaptureCapturerConfig
 {
 	FIntPoint OutputResolution;
 	bool	  bIsSRGB = false;
@@ -35,17 +37,17 @@ struct PIXELCAPTURE_API FPixelCaptureCapturerConfig
  * While the capture should be async it should only expect to work on one
  * frame at a time.
  */
-class PIXELCAPTURE_API FPixelCaptureCapturer
+class FPixelCaptureCapturer
 {
 public:
-	FPixelCaptureCapturer(FPixelCaptureCapturerConfig Config = {});
-	virtual ~FPixelCaptureCapturer();
+	UE_API FPixelCaptureCapturer(FPixelCaptureCapturerConfig Config = {});
+	UE_API virtual ~FPixelCaptureCapturer();
 
 	/**
 	 * Called when an input frame needs capturing.
 	 * @param InputFrame The input frame to be captured.
 	 */
-	void Capture(const IPixelCaptureInputFrame& InputFrame);
+	UE_API void Capture(const IPixelCaptureInputFrame& InputFrame);
 
 	/**
 	 * Returns true if Initialize() has been called.
@@ -70,7 +72,7 @@ public:
 	 * Gets the output frame from the output buffer.
 	 * @return The output data of this process.
 	 */
-	TSharedPtr<IPixelCaptureOutputFrame> ReadOutput();
+	UE_API TSharedPtr<IPixelCaptureOutputFrame> ReadOutput();
 
 	/**
 	 * Listen on this to be notified when the capture process completes for each input.
@@ -97,7 +99,7 @@ protected:
 	 * @param InputWidth The pixel count of the input frame width
 	 * @param InputHeight The pixel count of the input frame height
 	 */
-	virtual void Initialize(int32 InputWidth, int32 InputHeight);
+	UE_API virtual void Initialize(int32 InputWidth, int32 InputHeight);
 
 	/**
 	 * Implement this to create a buffer for the output.
@@ -145,7 +147,7 @@ protected:
 	 */
 	UE_DEPRECATED(5.6, "EndProcess(void) has been deprecated. Please use EndProcess(TSharedPtr<IPixelCaptureOutputFrame> OutputBuffer) to ensure the metadata is updated for the correct output buffer.")
 	void EndProcess() {}
-	void EndProcess(TSharedPtr<IPixelCaptureOutputFrame> OutputBuffer);
+	UE_API void EndProcess(TSharedPtr<IPixelCaptureOutputFrame> OutputBuffer);
 
 protected:
 	FPixelCaptureCapturerConfig Config;
@@ -165,6 +167,8 @@ private:
 	mutable FCriticalSection					 OutputBuffersMutex;
 	TArray<TSharedPtr<IPixelCaptureOutputFrame>> OutputBuffers;
 
-	void InitMetadata(FPixelCaptureFrameMetadata InputMetadata, TSharedPtr<IPixelCaptureOutputFrame> OutputBuffer);
-	void FinalizeMetadata(TSharedPtr<IPixelCaptureOutputFrame> OutputBuffer);
+	UE_API void InitMetadata(FPixelCaptureFrameMetadata InputMetadata, TSharedPtr<IPixelCaptureOutputFrame> OutputBuffer);
+	UE_API void FinalizeMetadata(TSharedPtr<IPixelCaptureOutputFrame> OutputBuffer);
 };
+
+#undef UE_API
diff --git a/Engine/Plugins/Media/PixelCapture/Source/PixelCapture/Public/PixelCaptureCapturerI420.h b/Engine/Plugins/Media/PixelCapture/Source/PixelCapture/Public/PixelCaptureCapturerI420.h
--- a/Engine/Plugins/Media/PixelCapture/Source/PixelCapture/Public/PixelCaptureCapturerI420.h
+++ b/Engine/Plugins/Media/PixelCapture/Source/PixelCapture/Public/PixelCaptureCapturerI420.h
@@ -4,23 +4,27 @@
 
 #include "PixelCaptureCapturer.h"
 
+#define UE_API PIXELCAPTURE_API
+
 /**
  * A basic capturer that will copy I420 frames.
  * Input: FPixelCaptureInputFrameI420
  * Output: FPixelCaptureOutputFrameI420
  */
-class PIXELCAPTURE_API FPixelCaptureCapturerI420 : public FPixelCaptureCapturer
+class FPixelCaptureCapturerI420 : public FPixelCaptureCapturer
 {
 public:
-	static TSharedPtr<FPixelCaptureCapturerI420> Create(FPixelCaptureCapturerConfig Config = {});
+	static UE_API TSharedPtr<FPixelCaptureCapturerI420> Create(FPixelCaptureCapturerConfig Config = {});
 	virtual ~FPixelCaptureCapturerI420() = default;
 
 protected:
 	virtual FString					  GetCapturerName() const override { return "I420 Copy"; }
-	virtual void					  Initialize(int32 InputWidth, int32 InputHeight) override;
-	virtual IPixelCaptureOutputFrame* CreateOutputBuffer(int32 InputWidth, int32 InputHeight) override;
-	virtual void					  BeginProcess(const IPixelCaptureInputFrame& InputFrame, TSharedPtr<IPixelCaptureOutputFrame> OutputBuffer) override;
+	UE_API virtual void					  Initialize(int32 InputWidth, int32 InputHeight) override;
+	UE_API virtual IPixelCaptureOutputFrame* CreateOutputBuffer(int32 InputWidth, int32 InputHeight) override;
+	UE_API virtual void					  BeginProcess(const IPixelCaptureInputFrame& InputFrame, TSharedPtr<IPixelCaptureOutputFrame> OutputBuffer) override;
 
 private:
-	FPixelCaptureCapturerI420(FPixelCaptureCapturerConfig& Config);
+	UE_API FPixelCaptureCapturerI420(FPixelCaptureCapturerConfig& Config);
 };
+
+#undef UE_API
diff --git a/Engine/Plugins/Media/PixelCapture/Source/PixelCapture/Public/PixelCaptureCapturerI420ToRHI.h b/Engine/Plugins/Media/PixelCapture/Source/PixelCapture/Public/PixelCaptureCapturerI420ToRHI.h
--- a/Engine/Plugins/Media/PixelCapture/Source/PixelCapture/Public/PixelCaptureCapturerI420ToRHI.h
+++ b/Engine/Plugins/Media/PixelCapture/Source/PixelCapture/Public/PixelCaptureCapturerI420ToRHI.h
@@ -4,25 +4,29 @@
 
 #include "PixelCaptureCapturer.h"
 
+#define UE_API PIXELCAPTURE_API
+
 /**
  * A basic capturer that will capture I420 frames to native RHI textures.
  * Input: FPixelCaptureInputFrameI420
  * Output: FPixelCaptureOutputFrameRHI
  */
-class PIXELCAPTURE_API FPixelCaptureCapturerI420ToRHI : public FPixelCaptureCapturer, public TSharedFromThis<FPixelCaptureCapturerI420ToRHI>
+class FPixelCaptureCapturerI420ToRHI : public FPixelCaptureCapturer, public TSharedFromThis<FPixelCaptureCapturerI420ToRHI>
 {
 public:
-	static TSharedPtr<FPixelCaptureCapturerI420ToRHI> Create(FPixelCaptureCapturerConfig Config = {});
-	virtual ~FPixelCaptureCapturerI420ToRHI();
+	static UE_API TSharedPtr<FPixelCaptureCapturerI420ToRHI> Create(FPixelCaptureCapturerConfig Config = {});
+	UE_API virtual ~FPixelCaptureCapturerI420ToRHI();
 
 protected:
 	virtual FString GetCapturerName() const override { return "I420ToRHI"; }
-	virtual void Initialize(int32 InputWidth, int32 InputHeight) override;
-	virtual IPixelCaptureOutputFrame* CreateOutputBuffer(int32 InputWidth, int32 InputHeight) override;
-	virtual void BeginProcess(const IPixelCaptureInputFrame& InputFrame, TSharedPtr<IPixelCaptureOutputFrame> OutputBuffer) override;
+	UE_API virtual void Initialize(int32 InputWidth, int32 InputHeight) override;
+	UE_API virtual IPixelCaptureOutputFrame* CreateOutputBuffer(int32 InputWidth, int32 InputHeight) override;
+	UE_API virtual void BeginProcess(const IPixelCaptureInputFrame& InputFrame, TSharedPtr<IPixelCaptureOutputFrame> OutputBuffer) override;
 
 private:
 	uint8_t* ARGBBuffer = nullptr;
 
-	FPixelCaptureCapturerI420ToRHI(FPixelCaptureCapturerConfig& Config);
+	UE_API FPixelCaptureCapturerI420ToRHI(FPixelCaptureCapturerConfig& Config);
 };
+
+#undef UE_API
diff --git a/Engine/Plugins/Media/PixelCapture/Source/PixelCapture/Public/PixelCaptureCapturerMultiFormat.h b/Engine/Plugins/Media/PixelCapture/Source/PixelCapture/Public/PixelCaptureCapturerMultiFormat.h
--- a/Engine/Plugins/Media/PixelCapture/Source/PixelCapture/Public/PixelCaptureCapturerMultiFormat.h
+++ b/Engine/Plugins/Media/PixelCapture/Source/PixelCapture/Public/PixelCaptureCapturerMultiFormat.h
@@ -5,13 +5,15 @@
 #include "IPixelCaptureCapturerSource.h"
 #include "PixelCaptureCapturerLayered.h"
 
+#define UE_API PIXELCAPTURE_API
+
 /**
  * A capturer that contains multiple formats of multi layer capture processes.
  * Feed it a IPixelCaptureCapturerSource that will create the appropriate base capturers.
  * Input: User defined
  * Output: Capturer defined
  */
-class PIXELCAPTURE_API FPixelCaptureCapturerMultiFormat : public TSharedFromThis<FPixelCaptureCapturerMultiFormat>
+class FPixelCaptureCapturerMultiFormat : public TSharedFromThis<FPixelCaptureCapturerMultiFormat>
 {
 public:
 	/**
@@ -19,7 +21,7 @@ public:
 	 * @param InCapturerSource A source for capturers for each layer.
 	 * @param InOutputResolutions An optional parameter containing the output resolutions if known.
 	 */
-	static TSharedPtr<FPixelCaptureCapturerMultiFormat> Create(IPixelCaptureCapturerSource* InCapturerSource, TArray<FIntPoint> InOutputResolutions = {});
+	static UE_API TSharedPtr<FPixelCaptureCapturerMultiFormat> Create(IPixelCaptureCapturerSource* InCapturerSource, TArray<FIntPoint> InOutputResolutions = {});
 
 	/**
 	 * Create a new multi-format multi-Layered Capturer.
@@ -27,9 +29,9 @@ public:
 	 * @param LayerScales A list of scales for each layer.
 	 */
 	UE_DEPRECATED(5.7, "Create taking in an array of layer scales has been deprecated. You no longer need to pass in the scales during construction!")
-	static TSharedPtr<FPixelCaptureCapturerMultiFormat> Create(IPixelCaptureCapturerSource* InCapturerSource, TArray<float> LayerScales);
+	static UE_API TSharedPtr<FPixelCaptureCapturerMultiFormat> Create(IPixelCaptureCapturerSource* InCapturerSource, TArray<float> LayerScales);
 
-	virtual ~FPixelCaptureCapturerMultiFormat();
+	UE_API virtual ~FPixelCaptureCapturerMultiFormat();
 
 	/**
 	 * Gets the number of layers in the multi-layered capturers.
@@ -60,14 +62,14 @@ public:
 	/**
 	 * Begins the capture process of a given frame.
 	 */
-	void Capture(const IPixelCaptureInputFrame& SourceFrame);
+	UE_API void Capture(const IPixelCaptureInputFrame& SourceFrame);
 
 	/**
 	 * Sets up a capture pipeline for the given destination format. No effect if the
 	 * pipeline already exists.
 	 * @param Format The destination format for the requested pipeline.
 	 */
-	void AddOutputFormat(int32 Format);
+	UE_API void AddOutputFormat(int32 Format);
 
 	/**
 	 * Requests the output in a specific format. If this is the first request for
@@ -78,7 +80,7 @@ public:
 	 * @param LayerSize The resolution of the layer we want to get the output of.
 	 * @return The final buffer of a given format if it exists or null.
 	 */
-	TSharedPtr<IPixelCaptureOutputFrame> RequestFormat(int32 Format, FIntPoint LayerSize);
+	UE_API TSharedPtr<IPixelCaptureOutputFrame> RequestFormat(int32 Format, FIntPoint LayerSize);
 
 	/**
 	 * Requests the output in a specific format. If this is the first request for
@@ -102,7 +104,7 @@ public:
 	 * @param MaxWaitTime The max number of milliseconds to wait for a frame. Default is 5 seconds.
 	 * @return The final buffer of a given format or null in case of timeout or the capturer has been disconnected.
 	 */
-	TSharedPtr<IPixelCaptureOutputFrame> WaitForFormat(int32 Format, FIntPoint LayerSize, uint32 MaxWaitTime = 5000);
+	UE_API TSharedPtr<IPixelCaptureOutputFrame> WaitForFormat(int32 Format, FIntPoint LayerSize, uint32 MaxWaitTime = 5000);
 
 	/**
 	 * Like RequestFormat except if the format does not exist it will add it and then
@@ -121,7 +123,7 @@ public:
 	 * Call to notify this capturer that it has been disconnected and no more frames
 	 * will be captured and any waiting for format calls should stop waiting.
 	 */
-	void OnDisconnected();
+	UE_API void OnDisconnected();
 
 	/**
 	 * Listen on this to be notified when a frame completes all capture formats/layers.
@@ -130,13 +132,13 @@ public:
 	FOnComplete OnComplete;
 
 protected:
-	FPixelCaptureCapturerMultiFormat(IPixelCaptureCapturerSource* InCapturerSource, TArray<FIntPoint> InOutputResolutions);
-	void OnCaptureFormatComplete(int32 Format);
+	UE_API FPixelCaptureCapturerMultiFormat(IPixelCaptureCapturerSource* InCapturerSource, TArray<FIntPoint> InOutputResolutions);
+	UE_API void OnCaptureFormatComplete(int32 Format);
 
-	FEvent* GetEventForFormat(int32 Format);
-	void	CheckFormatEvent(int32 Format);
-	void	FreeEvent(int32 Format, FEvent* Event);
-	void	FlushWaitingEvents();
+	UE_API FEvent* GetEventForFormat(int32 Format);
+	UE_API void	CheckFormatEvent(int32 Format);
+	UE_API void	FreeEvent(int32 Format, FEvent* Event);
+	UE_API void	FlushWaitingEvents();
 
 	IPixelCaptureCapturerSource* CapturerSource;
 
@@ -152,3 +154,5 @@ protected:
 
 	bool bDisconnected = false;
 };
+
+#undef UE_API
diff --git a/Engine/Plugins/Media/PixelCapture/Source/PixelCapture/Public/PixelCaptureCapturerNV12ToRHI.h b/Engine/Plugins/Media/PixelCapture/Source/PixelCapture/Public/PixelCaptureCapturerNV12ToRHI.h
--- a/Engine/Plugins/Media/PixelCapture/Source/PixelCapture/Public/PixelCaptureCapturerNV12ToRHI.h
+++ b/Engine/Plugins/Media/PixelCapture/Source/PixelCapture/Public/PixelCaptureCapturerNV12ToRHI.h
@@ -4,25 +4,29 @@
 
 #include "PixelCaptureCapturer.h"
 
+#define UE_API PIXELCAPTURE_API
+
 /**
  * A basic capturer that will capture NV12 frames to native RHI textures.
  * Input: FPixelCaptureInputFrameNV12
  * Output: FPixelCaptureOutputFrameRHI
  */
-class PIXELCAPTURE_API FPixelCaptureCapturerNV12ToRHI : public FPixelCaptureCapturer, public TSharedFromThis<FPixelCaptureCapturerNV12ToRHI>
+class FPixelCaptureCapturerNV12ToRHI : public FPixelCaptureCapturer, public TSharedFromThis<FPixelCaptureCapturerNV12ToRHI>
 {
 public:
-	static TSharedPtr<FPixelCaptureCapturerNV12ToRHI> Create(FPixelCaptureCapturerConfig Config = {});
-	virtual ~FPixelCaptureCapturerNV12ToRHI();
+	static UE_API TSharedPtr<FPixelCaptureCapturerNV12ToRHI> Create(FPixelCaptureCapturerConfig Config = {});
+	UE_API virtual ~FPixelCaptureCapturerNV12ToRHI();
 
 protected:
 	virtual FString GetCapturerName() const override { return "NV12ToRHI"; }
-	virtual void Initialize(int32 InputWidth, int32 InputHeight) override;
-	virtual IPixelCaptureOutputFrame* CreateOutputBuffer(int32 InputWidth, int32 InputHeight) override;
-	virtual void BeginProcess(const IPixelCaptureInputFrame& InputFrame, TSharedPtr<IPixelCaptureOutputFrame> OutputBuffer) override;
+	UE_API virtual void Initialize(int32 InputWidth, int32 InputHeight) override;
+	UE_API virtual IPixelCaptureOutputFrame* CreateOutputBuffer(int32 InputWidth, int32 InputHeight) override;
+	UE_API virtual void BeginProcess(const IPixelCaptureInputFrame& InputFrame, TSharedPtr<IPixelCaptureOutputFrame> OutputBuffer) override;
 
 private:
 	uint8_t* R8Buffer = nullptr;
 
-	FPixelCaptureCapturerNV12ToRHI(FPixelCaptureCapturerConfig& Config);
+	UE_API FPixelCaptureCapturerNV12ToRHI(FPixelCaptureCapturerConfig& Config);
 };
+
+#undef UE_API
diff --git a/Engine/Plugins/Media/PixelCapture/Source/PixelCapture/Public/PixelCaptureCapturerRHI.h b/Engine/Plugins/Media/PixelCapture/Source/PixelCapture/Public/PixelCaptureCapturerRHI.h
--- a/Engine/Plugins/Media/PixelCapture/Source/PixelCapture/Public/PixelCaptureCapturerRHI.h
+++ b/Engine/Plugins/Media/PixelCapture/Source/PixelCapture/Public/PixelCaptureCapturerRHI.h
@@ -5,15 +5,17 @@
 #include "PixelCaptureCapturer.h"
 #include "RHI.h"
 
+#define UE_API PIXELCAPTURE_API
+
 /**
  * A basic capturer that will copy RHI texture frames.
  * Input: FPixelCaptureInputFrameRHI
  * Output: FPixelCaptureOutputFrameRHI
  */
-class PIXELCAPTURE_API FPixelCaptureCapturerRHI : public FPixelCaptureCapturer, public TSharedFromThis<FPixelCaptureCapturerRHI>
+class FPixelCaptureCapturerRHI : public FPixelCaptureCapturer, public TSharedFromThis<FPixelCaptureCapturerRHI>
 {
 public:
-	static TSharedPtr<FPixelCaptureCapturerRHI> Create(FPixelCaptureCapturerConfig Config = {});
+	static UE_API TSharedPtr<FPixelCaptureCapturerRHI> Create(FPixelCaptureCapturerConfig Config = {});
 
 	UE_DEPRECATED(5.7, "Create taking a scale has been deprecated. Please use Create that takes an FPixelCaptureCaptureConfig instead")
 	static TSharedPtr<FPixelCaptureCapturerRHI> Create(float InScale) { return nullptr; }
@@ -22,13 +24,15 @@ public:
 
 protected:
 	virtual FString					  GetCapturerName() const override { return "RHI Copy"; }
-	virtual IPixelCaptureOutputFrame* CreateOutputBuffer(int32 InputWidth, int32 InputHeight) override;
-	virtual void					  BeginProcess(const IPixelCaptureInputFrame& InputFrame, TSharedPtr<IPixelCaptureOutputFrame> OutputBuffer) override;
+	UE_API virtual IPixelCaptureOutputFrame* CreateOutputBuffer(int32 InputWidth, int32 InputHeight) override;
+	UE_API virtual void					  BeginProcess(const IPixelCaptureInputFrame& InputFrame, TSharedPtr<IPixelCaptureOutputFrame> OutputBuffer) override;
 
 private:
 	ERHIInterfaceType RHIType;
 
-	FPixelCaptureCapturerRHI(FPixelCaptureCapturerConfig& Config);
-	void CheckComplete(TSharedPtr<IPixelCaptureOutputFrame> OutputBuffer, FGPUFenceRHIRef Fence);
-	void OnRHIStageComplete(TSharedPtr<IPixelCaptureOutputFrame> OutputBuffer);
+	UE_API FPixelCaptureCapturerRHI(FPixelCaptureCapturerConfig& Config);
+	UE_API void CheckComplete(TSharedPtr<IPixelCaptureOutputFrame> OutputBuffer, FGPUFenceRHIRef Fence);
+	UE_API void OnRHIStageComplete(TSharedPtr<IPixelCaptureOutputFrame> OutputBuffer);
 };
+
+#undef UE_API
diff --git a/Engine/Plugins/Media/PixelCapture/Source/PixelCapture/Public/PixelCaptureCapturerRHINoCopy.h b/Engine/Plugins/Media/PixelCapture/Source/PixelCapture/Public/PixelCaptureCapturerRHINoCopy.h
--- a/Engine/Plugins/Media/PixelCapture/Source/PixelCapture/Public/PixelCaptureCapturerRHINoCopy.h
+++ b/Engine/Plugins/Media/PixelCapture/Source/PixelCapture/Public/PixelCaptureCapturerRHINoCopy.h
@@ -5,6 +5,8 @@
 #include "PixelCaptureCapturer.h"
 #include "RHI.h"
 
+#define UE_API PIXELCAPTURE_API
+
 /**
  * A basic capturer for receiving copy-safe RHI Texture frames. This pixel capturer will shortcut the need to copy texture is the input
  * and output frame sizes match
@@ -12,13 +14,13 @@
  * Input: FPixelCaptureInputFrameRHI
  * Output: FPixelCaptureOutputFrameRHI
  */
-class PIXELCAPTURE_API FPixelCaptureCapturerRHINoCopy : public FPixelCaptureCapturer, public TSharedFromThis<FPixelCaptureCapturerRHINoCopy>
+class FPixelCaptureCapturerRHINoCopy : public FPixelCaptureCapturer, public TSharedFromThis<FPixelCaptureCapturerRHINoCopy>
 {
 public:
 	/**
 	 * Creates a new capturer that assigns the output frame to the input frame. NOTE: This means this capturer doesn't support resolution or format changes
 	 */
-	static TSharedPtr<FPixelCaptureCapturerRHINoCopy> Create(FPixelCaptureCapturerConfig Config = {});
+	static UE_API TSharedPtr<FPixelCaptureCapturerRHINoCopy> Create(FPixelCaptureCapturerConfig Config = {});
 
 	UE_DEPRECATED(5.7, "Create taking a scale has been deprecated. Please use Create that takes an FPixelCaptureCapturerConfig instead")
 	static TSharedPtr<FPixelCaptureCapturerRHINoCopy> Create(float InScale) { return nullptr; }
@@ -27,9 +29,11 @@ public:
 
 protected:
 	virtual FString					  GetCapturerName() const override { return "RHI No Copy"; }
-	virtual IPixelCaptureOutputFrame* CreateOutputBuffer(int32 InputWidth, int32 InputHeight) override;
-	virtual void					  BeginProcess(const IPixelCaptureInputFrame& InputFrame, TSharedPtr<IPixelCaptureOutputFrame> OutputBuffer) override;
+	UE_API virtual IPixelCaptureOutputFrame* CreateOutputBuffer(int32 InputWidth, int32 InputHeight) override;
+	UE_API virtual void					  BeginProcess(const IPixelCaptureInputFrame& InputFrame, TSharedPtr<IPixelCaptureOutputFrame> OutputBuffer) override;
 
 private:
-	FPixelCaptureCapturerRHINoCopy(FPixelCaptureCapturerConfig& Config);
-};
\ No newline at end of file
+	UE_API FPixelCaptureCapturerRHINoCopy(FPixelCaptureCapturerConfig& Config);
+};
+
+#undef UE_API
diff --git a/Engine/Plugins/Media/PixelCapture/Source/PixelCapture/Public/PixelCaptureCapturerRHIRDG.h b/Engine/Plugins/Media/PixelCapture/Source/PixelCapture/Public/PixelCaptureCapturerRHIRDG.h
--- a/Engine/Plugins/Media/PixelCapture/Source/PixelCapture/Public/PixelCaptureCapturerRHIRDG.h
+++ b/Engine/Plugins/Media/PixelCapture/Source/PixelCapture/Public/PixelCaptureCapturerRHIRDG.h
@@ -5,15 +5,17 @@
 #include "PixelCaptureCapturer.h"
 #include "RHI.h"
 
+#define UE_API PIXELCAPTURE_API
+
 /**
  * A basic capturer that will copy RHI texture frames.
  * Input: FPixelCaptureInputFrameRHI
  * Output: FPixelCaptureOutputFrameRHI
  */
-class PIXELCAPTURE_API FPixelCaptureCapturerRHIRDG : public FPixelCaptureCapturer, public TSharedFromThis<FPixelCaptureCapturerRHIRDG>
+class FPixelCaptureCapturerRHIRDG : public FPixelCaptureCapturer, public TSharedFromThis<FPixelCaptureCapturerRHIRDG>
 {
 public:
-	static TSharedPtr<FPixelCaptureCapturerRHIRDG> Create(FPixelCaptureCapturerConfig Config = {});
+	static UE_API TSharedPtr<FPixelCaptureCapturerRHIRDG> Create(FPixelCaptureCapturerConfig Config = {});
 
 	UE_DEPRECATED(5.7, "Create taking a scale has been deprecated. Please use Create that takes an FPixelCaptureCapturerConfig instead")
 	static TSharedPtr<FPixelCaptureCapturerRHIRDG> Create(float InScale) { return nullptr; }
@@ -22,11 +24,13 @@ public:
 
 protected:
 	virtual FString					  GetCapturerName() const override { return "RHI RDG Copy"; }
-	virtual IPixelCaptureOutputFrame* CreateOutputBuffer(int32 InputWidth, int32 InputHeight) override;
-	virtual void					  BeginProcess(const IPixelCaptureInputFrame& InputFrame, TSharedPtr<IPixelCaptureOutputFrame> OutputBuffer) override;
+	UE_API virtual IPixelCaptureOutputFrame* CreateOutputBuffer(int32 InputWidth, int32 InputHeight) override;
+	UE_API virtual void					  BeginProcess(const IPixelCaptureInputFrame& InputFrame, TSharedPtr<IPixelCaptureOutputFrame> OutputBuffer) override;
 
 private:
-	FPixelCaptureCapturerRHIRDG(FPixelCaptureCapturerConfig& Config);
-	void CheckComplete();
-	void OnRHIStageComplete(TSharedPtr<IPixelCaptureOutputFrame> OutputBuffer);
+	UE_API FPixelCaptureCapturerRHIRDG(FPixelCaptureCapturerConfig& Config);
+	UE_API void CheckComplete();
+	UE_API void OnRHIStageComplete(TSharedPtr<IPixelCaptureOutputFrame> OutputBuffer);
 };
+
+#undef UE_API
diff --git a/Engine/Plugins/Media/PixelCapture/Source/PixelCapture/Public/PixelCaptureCapturerRHIToI420CPU.h b/Engine/Plugins/Media/PixelCapture/Source/PixelCapture/Public/PixelCaptureCapturerRHIToI420CPU.h
--- a/Engine/Plugins/Media/PixelCapture/Source/PixelCapture/Public/PixelCaptureCapturerRHIToI420CPU.h
+++ b/Engine/Plugins/Media/PixelCapture/Source/PixelCapture/Public/PixelCaptureCapturerRHIToI420CPU.h
@@ -6,38 +6,42 @@
 #include "RHI.h"
 #include "RHIGPUReadback.h"
 
+#define UE_API PIXELCAPTURE_API
+
 /**
  * A basic capturer that will capture RHI texture frames to I420 buffers
  * utilizing cpu functions. Involves CPU readback of GPU textures and processing
  * of that readback data. Input: FPixelCaptureInputFrameRHI Output:
  * FPixelCaptureOutputFrameI420
  */
-class PIXELCAPTURE_API FPixelCaptureCapturerRHIToI420CPU : public FPixelCaptureCapturer, public TSharedFromThis<FPixelCaptureCapturerRHIToI420CPU>
+class FPixelCaptureCapturerRHIToI420CPU : public FPixelCaptureCapturer, public TSharedFromThis<FPixelCaptureCapturerRHIToI420CPU>
 {
 public:
-	static TSharedPtr<FPixelCaptureCapturerRHIToI420CPU> Create(FPixelCaptureCapturerConfig Config = {});
+	static UE_API TSharedPtr<FPixelCaptureCapturerRHIToI420CPU> Create(FPixelCaptureCapturerConfig Config = {});
 
 	UE_DEPRECATED(5.7, "Create taking a scale has been deprecated. Please use Create that takes an FPixelCaptureCapturerConfig instead")
 	static TSharedPtr<FPixelCaptureCapturerRHIToI420CPU> Create(float InScale) { return nullptr; }
 
-	virtual ~FPixelCaptureCapturerRHIToI420CPU();
+	UE_API virtual ~FPixelCaptureCapturerRHIToI420CPU();
 
 protected:
 	virtual FString					  GetCapturerName() const override { return "RHIToI420CPU"; }
-	virtual void					  Initialize(int32 InputWidth, int32 InputHeight) override;
-	virtual IPixelCaptureOutputFrame* CreateOutputBuffer(int32 InputWidth, int32 InputHeight) override;
-	virtual void					  BeginProcess(const IPixelCaptureInputFrame& InputFrame, TSharedPtr<IPixelCaptureOutputFrame> OutputBuffer) override;
+	UE_API virtual void					  Initialize(int32 InputWidth, int32 InputHeight) override;
+	UE_API virtual IPixelCaptureOutputFrame* CreateOutputBuffer(int32 InputWidth, int32 InputHeight) override;
+	UE_API virtual void					  BeginProcess(const IPixelCaptureInputFrame& InputFrame, TSharedPtr<IPixelCaptureOutputFrame> OutputBuffer) override;
 
 private:
 	FTextureRHIRef StagingTexture;
 
-	FPixelCaptureCapturerRHIToI420CPU(FPixelCaptureCapturerConfig& Config);
+	UE_API FPixelCaptureCapturerRHIToI420CPU(FPixelCaptureCapturerConfig& Config);
 
 #if 1
 	TSharedPtr<FRHIGPUTextureReadback> TextureReader;
-	void							   OnRHIStageComplete(TSharedPtr<IPixelCaptureOutputFrame> OutputBuffer);
-	void							   CheckComplete(TSharedPtr<IPixelCaptureOutputFrame> OutputBuffer);
+	UE_API void							   OnRHIStageComplete(TSharedPtr<IPixelCaptureOutputFrame> OutputBuffer);
+	UE_API void							   CheckComplete(TSharedPtr<IPixelCaptureOutputFrame> OutputBuffer);
 #else
 	FTextureRHIRef ReadbackTexture;
 #endif
 };
+
+#undef UE_API
diff --git a/Engine/Plugins/Media/PixelCapture/Source/PixelCapture/Public/PixelCaptureCapturerRHIToI420Compute.h b/Engine/Plugins/Media/PixelCapture/Source/PixelCapture/Public/PixelCaptureCapturerRHIToI420Compute.h
--- a/Engine/Plugins/Media/PixelCapture/Source/PixelCapture/Public/PixelCaptureCapturerRHIToI420Compute.h
+++ b/Engine/Plugins/Media/PixelCapture/Source/PixelCapture/Public/PixelCaptureCapturerRHIToI420Compute.h
@@ -5,27 +5,29 @@
 #include "PixelCaptureCapturer.h"
 #include "RHI.h"
 
+#define UE_API PIXELCAPTURE_API
+
 /**
  * A basic capturer that will capture RHI texture frames to I420 buffers utilizing a compute shader.
  * Involves compute shader reading and processing of GPU textures so should be faster than CPU variant.
  * Input: FPixelCaptureInputFrameRHI
  * Output: FPixelCaptureOutputFrameI420
  */
-class PIXELCAPTURE_API FPixelCaptureCapturerRHIToI420Compute : public FPixelCaptureCapturer, public TSharedFromThis<FPixelCaptureCapturerRHIToI420Compute>
+class FPixelCaptureCapturerRHIToI420Compute : public FPixelCaptureCapturer, public TSharedFromThis<FPixelCaptureCapturerRHIToI420Compute>
 {
 public:
-	static TSharedPtr<FPixelCaptureCapturerRHIToI420Compute> Create(FPixelCaptureCapturerConfig Config = {});
+	static UE_API TSharedPtr<FPixelCaptureCapturerRHIToI420Compute> Create(FPixelCaptureCapturerConfig Config = {});
 
 	UE_DEPRECATED(5.7, "Create taking a scale has been deprecated. Please use Create that takes an FPixelCaptureCapturerConfig instead")
 	static TSharedPtr<FPixelCaptureCapturerRHIToI420Compute> Create(float InScale) { return nullptr; }
 
-	virtual ~FPixelCaptureCapturerRHIToI420Compute();
+	UE_API virtual ~FPixelCaptureCapturerRHIToI420Compute();
 
 protected:
 	virtual FString					  GetCapturerName() const override { return "RHIToI420Compute"; }
-	virtual void					  Initialize(int32 InputWidth, int32 InputHeight) override;
-	virtual IPixelCaptureOutputFrame* CreateOutputBuffer(int32 InputWidth, int32 InputHeight) override;
-	virtual void					  BeginProcess(const IPixelCaptureInputFrame& InputFrame, TSharedPtr<IPixelCaptureOutputFrame> OutputBuffer) override;
+	UE_API virtual void					  Initialize(int32 InputWidth, int32 InputHeight) override;
+	UE_API virtual IPixelCaptureOutputFrame* CreateOutputBuffer(int32 InputWidth, int32 InputHeight) override;
+	UE_API virtual void					  BeginProcess(const IPixelCaptureInputFrame& InputFrame, TSharedPtr<IPixelCaptureOutputFrame> OutputBuffer) override;
 
 private:
 	// dimensions of the texures
@@ -56,7 +58,9 @@ private:
 	int32 UStride = 0;
 	int32 VStride = 0;
 
-	FPixelCaptureCapturerRHIToI420Compute(FPixelCaptureCapturerConfig& Config);
-	void OnRHIStageComplete(TSharedPtr<IPixelCaptureOutputFrame> OutputBuffer);
-	void CleanUp();
+	UE_API FPixelCaptureCapturerRHIToI420Compute(FPixelCaptureCapturerConfig& Config);
+	UE_API void OnRHIStageComplete(TSharedPtr<IPixelCaptureOutputFrame> OutputBuffer);
+	UE_API void CleanUp();
 };
+
+#undef UE_API
diff --git a/Engine/Plugins/Media/PixelCapture/Source/PixelCapture/Public/PixelCaptureFrameMetadata.h b/Engine/Plugins/Media/PixelCapture/Source/PixelCapture/Public/PixelCaptureFrameMetadata.h
--- a/Engine/Plugins/Media/PixelCapture/Source/PixelCapture/Public/PixelCaptureFrameMetadata.h
+++ b/Engine/Plugins/Media/PixelCapture/Source/PixelCapture/Public/PixelCaptureFrameMetadata.h
@@ -4,9 +4,9 @@
 
 #include "CoreMinimal.h"
 
-class PIXELCAPTURE_API FPixelCaptureUserData{};
+class FPixelCaptureUserData{};
 
-struct PIXELCAPTURE_API FPixelCaptureFrameMetadata
+struct FPixelCaptureFrameMetadata
 {
 	// Identifier for the capture pipeline/process this frame took
 	FString ProcessName = "Unknown";
diff --git a/Engine/Plugins/Media/PixelCapture/Source/PixelCapture/Public/PixelCaptureInputFrameI420.h b/Engine/Plugins/Media/PixelCapture/Source/PixelCapture/Public/PixelCaptureInputFrameI420.h
--- a/Engine/Plugins/Media/PixelCapture/Source/PixelCapture/Public/PixelCaptureInputFrameI420.h
+++ b/Engine/Plugins/Media/PixelCapture/Source/PixelCapture/Public/PixelCaptureInputFrameI420.h
@@ -5,21 +5,25 @@
 #include "IPixelCaptureInputFrame.h"
 #include "PixelCaptureBufferI420.h"
 
+#define UE_API PIXELCAPTURE_API
+
 /**
  * A basic input frame for the Capture system that wraps a I420 buffer.
  */
-class PIXELCAPTURE_API FPixelCaptureInputFrameI420 : public IPixelCaptureInputFrame
+class FPixelCaptureInputFrameI420 : public IPixelCaptureInputFrame
 {
 public:
-	FPixelCaptureInputFrameI420(TSharedPtr<FPixelCaptureBufferI420> Buffer, TSharedPtr<FPixelCaptureUserData> UserData = nullptr);
+	UE_API FPixelCaptureInputFrameI420(TSharedPtr<FPixelCaptureBufferI420> Buffer, TSharedPtr<FPixelCaptureUserData> UserData = nullptr);
 	virtual ~FPixelCaptureInputFrameI420() = default;
 
-	virtual int32 GetType() const override;
-	virtual int32 GetWidth() const override;
-	virtual int32 GetHeight() const override;
+	UE_API virtual int32 GetType() const override;
+	UE_API virtual int32 GetWidth() const override;
+	UE_API virtual int32 GetHeight() const override;
 
-	TSharedPtr<FPixelCaptureBufferI420> GetBuffer() const;
+	UE_API TSharedPtr<FPixelCaptureBufferI420> GetBuffer() const;
 
 private:
 	TSharedPtr<FPixelCaptureBufferI420> I420Buffer;
 };
+
+#undef UE_API
diff --git a/Engine/Plugins/Media/PixelCapture/Source/PixelCapture/Public/PixelCaptureInputFrameNV12.h b/Engine/Plugins/Media/PixelCapture/Source/PixelCapture/Public/PixelCaptureInputFrameNV12.h
--- a/Engine/Plugins/Media/PixelCapture/Source/PixelCapture/Public/PixelCaptureInputFrameNV12.h
+++ b/Engine/Plugins/Media/PixelCapture/Source/PixelCapture/Public/PixelCaptureInputFrameNV12.h
@@ -5,21 +5,25 @@
 #include "IPixelCaptureInputFrame.h"
 #include "PixelCaptureBufferNV12.h"
 
+#define UE_API PIXELCAPTURE_API
+
 /**
  * A basic input frame for the Capture system that wraps a I420 buffer.
  */
-class PIXELCAPTURE_API FPixelCaptureInputFrameNV12 : public IPixelCaptureInputFrame
+class FPixelCaptureInputFrameNV12 : public IPixelCaptureInputFrame
 {
 public:
-	FPixelCaptureInputFrameNV12(TSharedPtr<FPixelCaptureBufferNV12> Buffer, TSharedPtr<FPixelCaptureUserData> UserData = nullptr);
+	UE_API FPixelCaptureInputFrameNV12(TSharedPtr<FPixelCaptureBufferNV12> Buffer, TSharedPtr<FPixelCaptureUserData> UserData = nullptr);
 	virtual ~FPixelCaptureInputFrameNV12() = default;
 
-	virtual int32 GetType() const override;
-	virtual int32 GetWidth() const override;
-	virtual int32 GetHeight() const override;
+	UE_API virtual int32 GetType() const override;
+	UE_API virtual int32 GetWidth() const override;
+	UE_API virtual int32 GetHeight() const override;
 
-	TSharedPtr<FPixelCaptureBufferNV12> GetBuffer() const;
+	UE_API TSharedPtr<FPixelCaptureBufferNV12> GetBuffer() const;
 
 private:
 	TSharedPtr<FPixelCaptureBufferNV12> NV12Buffer;
 };
+
+#undef UE_API
diff --git a/Engine/Plugins/Media/PixelCapture/Source/PixelCapture/Public/PixelCaptureInputFrameRHI.h b/Engine/Plugins/Media/PixelCapture/Source/PixelCapture/Public/PixelCaptureInputFrameRHI.h
--- a/Engine/Plugins/Media/PixelCapture/Source/PixelCapture/Public/PixelCaptureInputFrameRHI.h
+++ b/Engine/Plugins/Media/PixelCapture/Source/PixelCapture/Public/PixelCaptureInputFrameRHI.h
@@ -5,18 +5,22 @@
 #include "IPixelCaptureInputFrame.h"
 #include "RHI.h"
 
+#define UE_API PIXELCAPTURE_API
+
 /**
  * A basic input frame for the Capture system that wraps a RHI texture buffer.
  */
-class PIXELCAPTURE_API FPixelCaptureInputFrameRHI : public IPixelCaptureInputFrame
+class FPixelCaptureInputFrameRHI : public IPixelCaptureInputFrame
 {
 public:
-	FPixelCaptureInputFrameRHI(FTextureRHIRef InFrameTexture, TSharedPtr<FPixelCaptureUserData> UserData = nullptr);
+	UE_API FPixelCaptureInputFrameRHI(FTextureRHIRef InFrameTexture, TSharedPtr<FPixelCaptureUserData> UserData = nullptr);
 	virtual ~FPixelCaptureInputFrameRHI() = default;
 
-	virtual int32 GetType() const override;
-	virtual int32 GetWidth() const override;
-	virtual int32 GetHeight() const override;
+	UE_API virtual int32 GetType() const override;
+	UE_API virtual int32 GetWidth() const override;
+	UE_API virtual int32 GetHeight() const override;
 
 	FTextureRHIRef FrameTexture;
 };
+
+#undef UE_API
diff --git a/Engine/Plugins/Media/PixelCapture/Source/PixelCapture/Public/PixelCaptureOutputFrameI420.h b/Engine/Plugins/Media/PixelCapture/Source/PixelCapture/Public/PixelCaptureOutputFrameI420.h
--- a/Engine/Plugins/Media/PixelCapture/Source/PixelCapture/Public/PixelCaptureOutputFrameI420.h
+++ b/Engine/Plugins/Media/PixelCapture/Source/PixelCapture/Public/PixelCaptureOutputFrameI420.h
@@ -8,7 +8,7 @@
 /**
  * A basic output frame from the Capture system that wraps a I420 buffer.
  */
-class PIXELCAPTURE_API FPixelCaptureOutputFrameI420 : public IPixelCaptureOutputFrame
+class FPixelCaptureOutputFrameI420 : public IPixelCaptureOutputFrame
 {
 public:
 	FPixelCaptureOutputFrameI420(TSharedPtr<FPixelCaptureBufferI420> InI420Buffer)
diff --git a/Engine/Plugins/Media/PixelCapture/Source/PixelCapture/Public/PixelCaptureOutputFrameRHI.h b/Engine/Plugins/Media/PixelCapture/Source/PixelCapture/Public/PixelCaptureOutputFrameRHI.h
--- a/Engine/Plugins/Media/PixelCapture/Source/PixelCapture/Public/PixelCaptureOutputFrameRHI.h
+++ b/Engine/Plugins/Media/PixelCapture/Source/PixelCapture/Public/PixelCaptureOutputFrameRHI.h
@@ -8,7 +8,7 @@
 /**
  * A basic output frame from the Capture system that wraps a RHI texture buffer.
  */
-class PIXELCAPTURE_API FPixelCaptureOutputFrameRHI : public IPixelCaptureOutputFrame
+class FPixelCaptureOutputFrameRHI : public IPixelCaptureOutputFrame
 {
 public:
 	/**
diff --git a/Engine/Plugins/Media/PixelCapture/Source/PixelCaptureShaders/Public/RGBToYUVShader.h b/Engine/Plugins/Media/PixelCapture/Source/PixelCaptureShaders/Public/RGBToYUVShader.h
--- a/Engine/Plugins/Media/PixelCapture/Source/PixelCaptureShaders/Public/RGBToYUVShader.h
+++ b/Engine/Plugins/Media/PixelCapture/Source/PixelCaptureShaders/Public/RGBToYUVShader.h
@@ -7,6 +7,8 @@
 #include "RHICommandList.h"
 #include "RHIStaticStates.h"
 
+#define UE_API PIXELCAPTURESHADERS_API
+
 struct FRGBToYUVShaderParameters
 {
 	FTextureRHIRef SourceTexture;
@@ -19,8 +21,10 @@ struct FRGBToYUVShaderParameters
 	FUnorderedAccessViewRHIRef DestPlaneV;
 };
 
-class PIXELCAPTURESHADERS_API FRGBToYUVShader
+class FRGBToYUVShader
 {
 public:
-	static void Dispatch(FRHICommandListImmediate& RHICmdList, const FRGBToYUVShaderParameters& InParameters);
+	static UE_API void Dispatch(FRHICommandListImmediate& RHICmdList, const FRGBToYUVShaderParameters& InParameters);
 };
+
+#undef UE_API
diff --git a/Engine/Plugins/Media/PixelStreaming2/Source/PixelStreaming2/Internal/AudioCapturer.h b/Engine/Plugins/Media/PixelStreaming2/Source/PixelStreaming2/Internal/AudioCapturer.h
--- a/Engine/Plugins/Media/PixelStreaming2/Source/PixelStreaming2/Internal/AudioCapturer.h
+++ b/Engine/Plugins/Media/PixelStreaming2/Source/PixelStreaming2/Internal/AudioCapturer.h
@@ -7,19 +7,21 @@
 #include "SampleBuffer.h"
 #include "TickableTask.h"
 
+#define UE_API PIXELSTREAMING2_API
+
 namespace UE::PixelStreaming2
 {
 	class FAudioCapturer;
 
-	class PIXELSTREAMING2_API FAudioPatchMixer : public Audio::FPatchMixer
+	class FAudioPatchMixer : public Audio::FPatchMixer
 	{
 	public:
-		FAudioPatchMixer(uint8 NumChannels, uint32 SampleRate, float SampleSizeSeconds);
+		UE_API FAudioPatchMixer(uint8 NumChannels, uint32 SampleRate, float SampleSizeSeconds);
 		virtual ~FAudioPatchMixer() = default;
 
-		uint32 GetMaxBufferSize() const;
-		uint8  GetNumChannels() const;
-		uint32 GetSampleRate() const;
+		UE_API uint32 GetMaxBufferSize() const;
+		UE_API uint8  GetNumChannels() const;
+		UE_API uint32 GetSampleRate() const;
 
 	protected:
 		uint8  NumChannels;
@@ -27,13 +29,13 @@ namespace UE::PixelStreaming2
 		float  SampleSizeSeconds;
 	};
 
-	class PIXELSTREAMING2_API FPatchInputProxy : public IPixelStreaming2AudioProducer
+	class FPatchInputProxy : public IPixelStreaming2AudioProducer
 	{
 	public:
-		FPatchInputProxy(TSharedPtr<FAudioPatchMixer> InMixer);
-		virtual ~FPatchInputProxy() override;
+		UE_API FPatchInputProxy(TSharedPtr<FAudioPatchMixer> InMixer);
+		UE_API virtual ~FPatchInputProxy() override;
 
-		virtual void PushAudio(const float* AudioData, int32 InNumSamples, int32 InNumChannels, int32 InSampleRate) override;
+		UE_API virtual void PushAudio(const float* AudioData, int32 InNumSamples, int32 InNumChannels, int32 InSampleRate) override;
 
 	protected:
 		TSharedPtr<FAudioPatchMixer> Mixer;
@@ -43,16 +45,16 @@ namespace UE::PixelStreaming2
 		uint32						 SampleRate;
 	};
 
-	class PIXELSTREAMING2_API FMixAudioTask : public FPixelStreamingTickableTask
+	class FMixAudioTask : public FPixelStreamingTickableTask
 	{
 	public:
-		FMixAudioTask(FAudioCapturer* Capturer, TSharedPtr<FAudioPatchMixer> Mixer);
+		UE_API FMixAudioTask(FAudioCapturer* Capturer, TSharedPtr<FAudioPatchMixer> Mixer);
 
 		virtual ~FMixAudioTask() = default;
 
 		// Begin FPixelStreamingTickableTask
-		virtual void		   Tick(float DeltaMs) override;
-		virtual const FString& GetName() const override;
+		UE_API virtual void		   Tick(float DeltaMs) override;
+		UE_API virtual const FString& GetName() const override;
 		// End FPixelStreamingTickableTask
 
 	protected:
@@ -63,18 +65,18 @@ namespace UE::PixelStreaming2
 		TSharedPtr<FAudioPatchMixer> Mixer;
 	};
 
-	class PIXELSTREAMING2_API FAudioCapturer : public IPixelStreaming2AudioProducer
+	class FAudioCapturer : public IPixelStreaming2AudioProducer
 	{
 	public:
-		static TSharedPtr<FAudioCapturer> Create(const int InSampleRate = 48000, const int InNumChannels = 2, const float InSampleSizeInSeconds = 0.5f);
+		static UE_API TSharedPtr<FAudioCapturer> Create(const int InSampleRate = 48000, const int InNumChannels = 2, const float InSampleSizeInSeconds = 0.5f);
 		virtual ~FAudioCapturer() = default;
 
 		// Mixed audio input will push its audio to an FPatchInputProxy for mixing
-		TSharedPtr<FAudioProducer> CreateAudioProducer();
-		void					   CreateAudioProducer(Audio::FDeviceId AudioDeviceId);
-		void					   RemoveAudioProducer(Audio::FDeviceId AudioDeviceId);
+		UE_API TSharedPtr<FAudioProducer> CreateAudioProducer();
+		UE_API void					   CreateAudioProducer(Audio::FDeviceId AudioDeviceId);
+		UE_API void					   RemoveAudioProducer(Audio::FDeviceId AudioDeviceId);
 
-		virtual void PushAudio(const float* AudioData, int32 NumSamples, int32 NumChannels, int32 SampleRate) override;
+		UE_API virtual void PushAudio(const float* AudioData, int32 NumSamples, int32 NumChannels, int32 SampleRate) override;
 
 		/**
 		 * This is broadcast each time audio is captured. Tracks should bind to this and push the audio into the track
@@ -83,14 +85,14 @@ namespace UE::PixelStreaming2
 		FOnAudioBuffer OnAudioBuffer;
 
 	protected:
-		FAudioCapturer(const int SampleRate = 48000, const int NumChannels = 2, const float SampleSizeInSeconds = 0.5f);
+		UE_API FAudioCapturer(const int SampleRate = 48000, const int NumChannels = 2, const float SampleSizeInSeconds = 0.5f);
 
-		void OnDebugDumpAudioChanged(IConsoleVariable* Var);
-		void OnEnginePreExit();
-		void WriteDebugAudio();
+		UE_API void OnDebugDumpAudioChanged(IConsoleVariable* Var);
+		UE_API void OnEnginePreExit();
+		UE_API void WriteDebugAudio();
 		// Allow the mix audio task to call OnAudio.
 		friend class FMixAudioTask;
-		void OnAudio(const float* AudioData, int32 NumSamples, int32 NumChannels, int32 SampleRate);
+		UE_API void OnAudio(const float* AudioData, int32 NumSamples, int32 NumChannels, int32 SampleRate);
 
 	protected:
 		TSharedPtr<FAudioPatchMixer> Mixer;
@@ -105,4 +107,6 @@ namespace UE::PixelStreaming2
 
 		Audio::TSampleBuffer<int16_t> DebugDumpAudioBuffer;
 	};
-} // namespace UE::PixelStreaming2
\ No newline at end of file
+} // namespace UE::PixelStreaming2
+
+#undef UE_API
diff --git a/Engine/Plugins/Media/PixelStreaming2/Source/PixelStreaming2/Internal/AudioProducer.h b/Engine/Plugins/Media/PixelStreaming2/Source/PixelStreaming2/Internal/AudioProducer.h
--- a/Engine/Plugins/Media/PixelStreaming2/Source/PixelStreaming2/Internal/AudioProducer.h
+++ b/Engine/Plugins/Media/PixelStreaming2/Source/PixelStreaming2/Internal/AudioProducer.h
@@ -6,6 +6,8 @@
 #include "ISubmixBufferListener.h"
 #include "IPixelStreaming2AudioProducer.h"
 
+#define UE_API PIXELSTREAMING2_API
+
 namespace UE::PixelStreaming2
 {
 	class FPatchInputProxy;
@@ -14,26 +16,28 @@ namespace UE::PixelStreaming2
 	 * An audio input capable of listening to UE submix's as well as receiving user audio via the PushAudio method.
 	 * Any received audio will be passed into the Parent's PushAudio method
 	 */
-	class PIXELSTREAMING2_API FAudioProducer : public ISubmixBufferListener, public IPixelStreaming2AudioProducer
+	class FAudioProducer : public ISubmixBufferListener, public IPixelStreaming2AudioProducer
 	{
 	public:
-		static TSharedPtr<FAudioProducer> Create(Audio::FDeviceId AudioDeviceId, TSharedPtr<FPatchInputProxy> PatchInput);
-		static TSharedPtr<FAudioProducer> Create(TSharedPtr<FPatchInputProxy> PatchInput);
+		static UE_API TSharedPtr<FAudioProducer> Create(Audio::FDeviceId AudioDeviceId, TSharedPtr<FPatchInputProxy> PatchInput);
+		static UE_API TSharedPtr<FAudioProducer> Create(TSharedPtr<FPatchInputProxy> PatchInput);
 		virtual ~FAudioProducer() = default;
 
 		// For users to manually push non-submix audio into
-		virtual void PushAudio(const float* AudioData, int32 NumSamples, int32 NumChannels, int32 SampleRate) override;
+		UE_API virtual void PushAudio(const float* AudioData, int32 NumSamples, int32 NumChannels, int32 SampleRate) override;
 
 		// ISubmixBufferListener interface
-		virtual void OnNewSubmixBuffer(const USoundSubmix* OwningSubmix, float* AudioData, int32 NumSamples, int32 NumChannels, const int32 SampleRate, double AudioClock) override;
+		UE_API virtual void OnNewSubmixBuffer(const USoundSubmix* OwningSubmix, float* AudioData, int32 NumSamples, int32 NumChannels, const int32 SampleRate, double AudioClock) override;
 
 		void ToggleMuted() { bIsMuted = !bIsMuted; }
 
 	protected:
-		FAudioProducer(TSharedPtr<FPatchInputProxy> PatchInput);
+		UE_API FAudioProducer(TSharedPtr<FPatchInputProxy> PatchInput);
 
 	protected:
 		TSharedPtr<FPatchInputProxy> PatchInput;
 		bool						 bIsMuted = false;
 	};
-} // namespace UE::PixelStreaming2
\ No newline at end of file
+} // namespace UE::PixelStreaming2
+
+#undef UE_API
diff --git a/Engine/Plugins/Media/PixelStreaming2/Source/PixelStreaming2/Internal/AudioSink.h b/Engine/Plugins/Media/PixelStreaming2/Source/PixelStreaming2/Internal/AudioSink.h
--- a/Engine/Plugins/Media/PixelStreaming2/Source/PixelStreaming2/Internal/AudioSink.h
+++ b/Engine/Plugins/Media/PixelStreaming2/Source/PixelStreaming2/Internal/AudioSink.h
@@ -4,22 +4,24 @@
 
 #include "IPixelStreaming2AudioSink.h"
 
+#define UE_API PIXELSTREAMING2_API
+
 namespace UE::PixelStreaming2
 {
-	class PIXELSTREAMING2_API FAudioSink : public IPixelStreaming2AudioSink
+	class FAudioSink : public IPixelStreaming2AudioSink
 	{
 	public:
 		// Note: destructor will call destroy on any attached audio consumers
-		virtual ~FAudioSink();
+		UE_API virtual ~FAudioSink();
 
-		virtual void AddAudioConsumer(const TWeakPtrVariant<IPixelStreaming2AudioConsumer>& AudioConsumer) override;
-		virtual void RemoveAudioConsumer(const TWeakPtrVariant<IPixelStreaming2AudioConsumer>& AudioConsumer) override;
+		UE_API virtual void AddAudioConsumer(const TWeakPtrVariant<IPixelStreaming2AudioConsumer>& AudioConsumer) override;
+		UE_API virtual void RemoveAudioConsumer(const TWeakPtrVariant<IPixelStreaming2AudioConsumer>& AudioConsumer) override;
 
-		bool HasAudioConsumers();
+		UE_API bool HasAudioConsumers();
 
-		void SetMuted(bool bIsMuted);
+		UE_API void SetMuted(bool bIsMuted);
 
-		void OnAudioData(int16_t* AudioData, uint32 NumFrames, uint32 NumChannels, uint32 SampleRate);
+		UE_API void OnAudioData(int16_t* AudioData, uint32 NumFrames, uint32 NumChannels, uint32 SampleRate);
 
 	protected:
 		TArray<TWeakPtrVariant<IPixelStreaming2AudioConsumer>> AudioConsumers;
@@ -28,4 +30,6 @@ namespace UE::PixelStreaming2
 		FCriticalSection AudioConsumersCS;
 		bool			 bIsMuted = false;
 	};
-} // namespace UE::PixelStreaming2
\ No newline at end of file
+} // namespace UE::PixelStreaming2
+
+#undef UE_API
diff --git a/Engine/Plugins/Media/PixelStreaming2/Source/PixelStreaming2/Internal/AudioSource.h b/Engine/Plugins/Media/PixelStreaming2/Source/PixelStreaming2/Internal/AudioSource.h
--- a/Engine/Plugins/Media/PixelStreaming2/Source/PixelStreaming2/Internal/AudioSource.h
+++ b/Engine/Plugins/Media/PixelStreaming2/Source/PixelStreaming2/Internal/AudioSource.h
@@ -4,18 +4,22 @@
 
 #include "Templates/SharedPointer.h"
 
+#define UE_API PIXELSTREAMING2_API
+
 namespace UE::PixelStreaming2
 {
-	class PIXELSTREAMING2_API FAudioSource
+	class FAudioSource
 	{
 	public:
 		virtual ~FAudioSource() = default;
 
-		virtual void OnAudioBuffer(const int16_t* AudioData, int32 NumSamples, int32 NumChannels, const int32 SampleRate);
+		UE_API virtual void OnAudioBuffer(const int16_t* AudioData, int32 NumSamples, int32 NumChannels, const int32 SampleRate);
 
-		void SetMuted(bool bIsMuted);
+		UE_API void SetMuted(bool bIsMuted);
 
 	protected:
 		bool bIsMuted = false;
 	};
-} // namespace UE::PixelStreaming2
\ No newline at end of file
+} // namespace UE::PixelStreaming2
+
+#undef UE_API
diff --git a/Engine/Plugins/Media/PixelStreaming2/Source/PixelStreaming2/Internal/Blueprints/PixelStreaming2MediaTexture.h b/Engine/Plugins/Media/PixelStreaming2/Source/PixelStreaming2/Internal/Blueprints/PixelStreaming2MediaTexture.h
--- a/Engine/Plugins/Media/PixelStreaming2/Source/PixelStreaming2/Internal/Blueprints/PixelStreaming2MediaTexture.h
+++ b/Engine/Plugins/Media/PixelStreaming2/Source/PixelStreaming2/Internal/Blueprints/PixelStreaming2MediaTexture.h
@@ -7,6 +7,8 @@
 
 #include "PixelStreaming2MediaTexture.generated.h"
 
+#define UE_API PIXELSTREAMING2_API
+
 namespace UE::PixelStreaming2
 {
 	class FPixelStreaming2MediaTextureResource;
@@ -15,31 +17,33 @@ namespace UE::PixelStreaming2
 /**
  * A Texture Object that can be used in materials etc. that takes updates from webrtc frames.
  */
-UCLASS(NotBlueprintType, NotBlueprintable, HideDropdown, HideCategories = (ImportSettings, Compression, Texture, Adjustments, Compositing, LevelOfDetail, Object), META = (DisplayName = "PixelStreaming Media Texture"))
-class PIXELSTREAMING2_API UPixelStreaming2MediaTexture : public UTexture2DDynamic, public IPixelStreaming2VideoConsumer
+UCLASS(MinimalAPI, NotBlueprintType, NotBlueprintable, HideDropdown, HideCategories = (ImportSettings, Compression, Texture, Adjustments, Compositing, LevelOfDetail, Object), META = (DisplayName = "PixelStreaming Media Texture"))
+class UPixelStreaming2MediaTexture : public UTexture2DDynamic, public IPixelStreaming2VideoConsumer
 {
 	GENERATED_UCLASS_BODY()
 
 protected:
 	// UObject overrides.
-	virtual void BeginDestroy() override;
-	virtual void GetResourceSizeEx(FResourceSizeEx& CumulativeResourceSize) override;
+	UE_API virtual void BeginDestroy() override;
+	UE_API virtual void GetResourceSizeEx(FResourceSizeEx& CumulativeResourceSize) override;
 
 	// UTexture implementation
-	virtual FTextureResource* CreateResource() override;
+	UE_API virtual FTextureResource* CreateResource() override;
 
 	// IPixelStreamingVideoConsumer implementation
-	virtual void ConsumeFrame(FTextureRHIRef Frame) override;
+	UE_API virtual void ConsumeFrame(FTextureRHIRef Frame) override;
 	virtual void OnConsumerAdded() override {}
 	virtual void OnConsumerRemoved() override {}
 
 private:
-	void InitializeResources();
+	UE_API void InitializeResources();
 
 	// updates the internal texture resource after each frame.
-	void UpdateTextureReference(FRHICommandList& RHICmdList, FTextureRHIRef Reference);
+	UE_API void UpdateTextureReference(FRHICommandList& RHICmdList, FTextureRHIRef Reference);
 
 	FCriticalSection RenderSyncContext;
 	// NOTE: This type has to be a raw ptr because UE will manage the lifetime of it
 	UE::PixelStreaming2::FPixelStreaming2MediaTextureResource* CurrentResource;
 };
+
+#undef UE_API
