[AVCodecs] Fix: Downgrade warning to log for build health

diff --git a/Engine/Plugins/Experimental/AVCodecs/AVCodecsCore/Source/AVCodecsCore/Private/AVResult.cpp b/Engine/Plugins/Experimental/AVCodecs/AVCodecsCore/Source/AVCodecsCore/Private/AVResult.cpp
--- a/Engine/Plugins/Experimental/AVCodecs/AVCodecsCore/Source/AVCodecsCore/Private/AVResult.cpp
+++ b/Engine/Plugins/Experimental/AVCodecs/AVCodecsCore/Source/AVCodecsCore/Private/AVResult.cpp
@@ -126,6 +126,10 @@ void FAVResult::Log() const
 		{
 			UE_LOG(LogAVCodecs, Verbose, TEXT("%s"), *ToString());
 		}
+		else
+		{
+			UE_LOG(LogAVCodecs, Log, TEXT("%s"), *ToString());
+		}
 	}
 	else
 	{
diff --git a/Engine/Plugins/Experimental/AVCodecs/NVCodecs/Source/NVDEC/Private/NVDEC.cpp b/Engine/Plugins/Experimental/AVCodecs/NVCodecs/Source/NVDEC/Private/NVDEC.cpp
--- a/Engine/Plugins/Experimental/AVCodecs/NVCodecs/Source/NVDEC/Private/NVDEC.cpp
+++ b/Engine/Plugins/Experimental/AVCodecs/NVCodecs/Source/NVDEC/Private/NVDEC.cpp
@@ -66,7 +66,7 @@ FNVDEC::FNVDEC()
 		}
 		else
 		{
-			FAVResult::Log(EAVResult::Warning, TEXT("Failed to get NVDEC dll handle. NVDEC module will not be available."), TEXT("NVDEC"));
+			FAVResult::Log(EAVResult::Unknown, TEXT("Failed to get NVDEC dll handle. NVDEC module will not be available."), TEXT("NVDEC"));
 		}
 	}
 	else
diff --git a/Engine/Plugins/Experimental/AVCodecs/NVCodecs/Source/NVENC/Private/NVENC.cpp b/Engine/Plugins/Experimental/AVCodecs/NVCodecs/Source/NVENC/Private/NVENC.cpp
--- a/Engine/Plugins/Experimental/AVCodecs/NVCodecs/Source/NVENC/Private/NVENC.cpp
+++ b/Engine/Plugins/Experimental/AVCodecs/NVCodecs/Source/NVENC/Private/NVENC.cpp
@@ -59,7 +59,7 @@ FNVENC::FNVENC()
 		}
 		else
 		{
-			FAVResult::Log(EAVResult::Warning, TEXT("Failed to get NVENC dll handle. NVENC module will not be available."), TEXT("NVENC"));
+			FAVResult::Log(EAVResult::Unknown, TEXT("Failed to get NVENC dll handle. NVENC module will not be available."), TEXT("NVENC"));
 		}
 	}
 	else
